Adding a commandlet to export an asset registry database to a sqlite database. The normal commandlet directory is in UnrealEd which doesn't reference the sqlite plugins, and rather than add those dependencies to UnrealEd, I made a separate plugin.

diff --git a/Engine/Plugins/Editor/AssetRegistryExport/AssetRegistryExport.uplugin b/Engine/Plugins/Editor/AssetRegistryExport/AssetRegistryExport.uplugin
new file mode 100644
--- /dev/null
+++ b/Engine/Plugins/Editor/AssetRegistryExport/AssetRegistryExport.uplugin
@@ -0,0 +1,30 @@
+{
+	"FileVersion": 3,
+	"Version": 1,
+	"VersionName": "1.0",
+	"FriendlyName": "Asset Registry Export",
+	"Description": "",
+	"Category": "Editor",
+	"CreatedBy": "Epic Games, Inc.",
+	"CreatedByURL": "http://epicgames.com",
+	"DocsURL": "",
+	"MarketplaceURL": "",
+	"SupportURL": "",
+	"EnabledByDefault": false,
+	"CanContainContent": false,
+	"IsBetaVersion": false,
+	"Installed": false,
+	"Modules": [
+		{
+			"Name": "AssetRegistryExport",
+			"Type": "Editor",
+			"LoadingPhase": "Default"
+		}
+	],
+	"Plugins": [
+		{
+			"Name": "SQLiteCore",
+			"Enabled": true
+		}
+	]
+}
\ No newline at end of file
diff --git a/Engine/Plugins/Editor/AssetRegistryExport/Source/AssetRegistryExport.Build.cs b/Engine/Plugins/Editor/AssetRegistryExport/Source/AssetRegistryExport.Build.cs
new file mode 100644
--- /dev/null
+++ b/Engine/Plugins/Editor/AssetRegistryExport/Source/AssetRegistryExport.Build.cs
@@ -0,0 +1,27 @@
+// Copyright Epic Games, Inc. All Rights Reserved.
+
+using UnrealBuildTool;
+
+public class AssetRegistryExport : ModuleRules
+{
+	public AssetRegistryExport(ReadOnlyTargetRules Target) : base(Target)
+	{
+		PublicDependencyModuleNames.AddRange(
+			new string[]
+			{
+				"Core",
+				"Engine"
+			}
+		);
+
+		PrivateDependencyModuleNames.AddRange(
+			new string[] {
+				"Core",
+				"CoreUObject",
+				"Engine",
+				"SQLiteCore",
+				"AssetRegistry",
+			}
+		);
+	}
+}
diff --git a/Engine/Plugins/Editor/AssetRegistryExport/Source/Private/AssetRegistryExportCommandlet.cpp b/Engine/Plugins/Editor/AssetRegistryExport/Source/Private/AssetRegistryExportCommandlet.cpp
new file mode 100644
--- /dev/null
+++ b/Engine/Plugins/Editor/AssetRegistryExport/Source/Private/AssetRegistryExportCommandlet.cpp
@@ -0,0 +1,312 @@
+#include "AssetRegistryExportCommandlet.h"
+
+#include "AssetData.h"
+#include "AssetRegistry/AssetRegistryState.h"
+#include "HAL/FileManager.h"
+#include "Misc/FileHelper.h"
+#include "Serialization/ArrayReader.h"
+#include "Serialization/MemoryReader.h"
+#include "Serialization/MemoryWriter.h"
+#include "SQLiteDatabase.h"
+
+DEFINE_LOG_CATEGORY(LogAssetRegistryExport);
+
+
+UAssetRegistryExportCommandlet::UAssetRegistryExportCommandlet(const FObjectInitializer& Initializer)
+	: Super(Initializer)
+{
+}
+
+int32 UAssetRegistryExportCommandlet::Main(const FString& CmdLineParams)
+{
+	FString OutputDatabaseName;
+	FString AssetRegistryFileName;
+
+	if (FParse::Value(*CmdLineParams, TEXT("Output="), OutputDatabaseName) == false ||
+		FParse::Value(*CmdLineParams, TEXT("Input="), AssetRegistryFileName) == false)
+	{
+		UE_LOG(LogAssetRegistryExport, Display, TEXT(""));
+		UE_LOG(LogAssetRegistryExport, Display, TEXT("Asset Registry Exporter"));
+		UE_LOG(LogAssetRegistryExport, Display, TEXT(""));
+		UE_LOG(LogAssetRegistryExport, Display, TEXT("Converts a given asset registry to a sqlite database. Each AssetClass gets its own"));
+		UE_LOG(LogAssetRegistryExport, Display, TEXT("table with the corresponding tags for the class as columns. Assets are added to their"));
+		UE_LOG(LogAssetRegistryExport, Display, TEXT("class's table, as well as a global Assets table."));
+		UE_LOG(LogAssetRegistryExport, Display, TEXT(""));
+		UE_LOG(LogAssetRegistryExport, Display, TEXT("Parameters:"));
+		UE_LOG(LogAssetRegistryExport, Display, TEXT(""));
+		UE_LOG(LogAssetRegistryExport, Display, TEXT("    -Output=<path/to/file>              Output Sqlite DB file (Required)"));
+		UE_LOG(LogAssetRegistryExport, Display, TEXT("    -Input=<path/to/file>               Input Asset Registry file (Required)"));
+		UE_LOG(LogAssetRegistryExport, Display, TEXT("    -Overwrite                          Allow overwrite of existing db file (Optional)."));
+		UE_LOG(LogAssetRegistryExport, Display, TEXT("                                        By default existing output files is a failure."));
+		return 1;
+	}
+
+	if (IFileManager::Get().FileExists(*OutputDatabaseName))
+	{
+		if (FParse::Param(*CmdLineParams, TEXT("Overwrite")) == false)
+		{
+			UE_LOG(LogAssetRegistryExport, Error, TEXT("Output database exists, delete or use -Overwrite. (%s)"), *OutputDatabaseName);
+			return 1;
+		}
+
+		UE_LOG(LogAssetRegistryExport, Display, TEXT("Output database exists, deleting."));
+		if (IFileManager::Get().Delete(*OutputDatabaseName) == false)
+		{
+			UE_LOG(LogAssetRegistryExport, Error, TEXT("Failed to delete output database %s."), *OutputDatabaseName);
+			return 1;
+		}
+	}
+
+	FSQLiteDatabase OutputDatabase;
+	if (OutputDatabase.Open(*OutputDatabaseName) == false)
+	{
+		UE_LOG(LogAssetRegistryExport, Error, TEXT("Failed to open sqlite database %s"), *OutputDatabaseName);
+		return 1;
+	}
+	UE_LOG(LogAssetRegistryExport, Display, TEXT("Creating Sqlite DB: %s"), *OutputDatabaseName);
+
+	FAssetRegistryState AssetRegistry;
+	{
+		FArrayReader SerializedAssetData;
+		if (FFileHelper::LoadFileToArray(SerializedAssetData, *AssetRegistryFileName))
+		{
+			FAssetRegistrySerializationOptions Options(UE::AssetRegistry::ESerializationTarget::ForDevelopment);
+			if (AssetRegistry.Serialize(SerializedAssetData, Options) == false)
+			{
+				UE_LOG(LogAssetRegistryExport, Error, TEXT("Failed to serialize asset registry %s"), *AssetRegistryFileName);
+			}
+		}
+		else
+		{
+			UE_LOG(LogAssetRegistryExport, Error, TEXT("Failed to read asset registry %s"), *AssetRegistryFileName);
+		}
+	}
+	UE_LOG(LogAssetRegistryExport, Display, TEXT("Opened asset registry: %s"), *AssetRegistryFileName);
+
+	// Set up some database properties for speed - we don't really care about ACID in this case
+	// since we're just dumping stuff and closing the database - we can always start over.
+	{
+		// don't bother waiting for disk to flush (let the OS handle it)
+		OutputDatabase.Execute(TEXT("PRAGMA synchronous=OFF;"));
+
+		// 1000 page cache
+		OutputDatabase.Execute(TEXT("PRAGMA cache_size=1000;"));
+
+		// large-ish page sizes for modern disks
+		OutputDatabase.Execute(TEXT("PRAGMA page_size=65535;"));
+
+		// hold the file lock the whole time.
+		OutputDatabase.Execute(TEXT("PRAGMA locking_mode=EXCLUSIVE;"));
+	}
+
+
+	struct FClassInfo
+	{
+		FName Name;
+		TSet<FName> Tags;
+
+		// TagName -> Name that has been sanitized for use as a column name.
+		TMap<FName, FString> SanitizedTagNames;
+
+		// TagName -> index for binding in the insert statement.
+		TMap<FName, int32> TagBindingIndices;
+
+		FSQLitePreparedStatement InsertStatement;
+	};
+
+	TMap<FName, FClassInfo> ClassInfos;
+	int32 AssetCount = 0;
+
+	// Find the tags for each class. Assets might not fill all tags that might exist on a
+	// given class, so we have to get the union of all possibles.
+	AssetRegistry.EnumerateAllAssets(TSet<FName>(), [&](const FAssetData& AssetData)
+	{
+		FClassInfo& ClassInfo = ClassInfos.FindOrAdd(AssetData.AssetClass);
+		ClassInfo.Name = AssetData.AssetClass;
+
+		AssetData.TagsAndValues.ForEach([&ClassInfo](TPair<FName, FAssetTagValueRef> Pair)
+		{
+			bool bAlreadySet = false;
+			ClassInfo.Tags.Add(Pair.Key, &bAlreadySet);
+
+			if (bAlreadySet == false)
+			{
+				FString& SanitizedTagName = ClassInfo.SanitizedTagNames.FindOrAdd(Pair.Key, Pair.Key.ToString());
+
+				// Only alphanumeric and _ are allowed, just make anything else a "_".
+				// Num() - 1 because FStrings store the null term in the array
+				auto& CharArray = SanitizedTagName.GetCharArray();
+				for (int i = 0; i < CharArray.Num() - 1; i++)
+				{
+					if (TChar<TCHAR>::IsAlnum(CharArray[i]) == false)
+					{
+						CharArray[i] = TEXT('_');
+					}
+				}
+			}
+		});
+
+		AssetCount++;
+		return true;
+	});
+
+	//
+	// Create the tables and insert statements for each class
+	//
+	for (TPair<FName, FClassInfo>& Pair : ClassInfos)
+	{
+		FClassInfo& Class = Pair.Value;
+
+		FString CreateTable = TEXT("CREATE TABLE IF NOT EXISTS ") + Class.Name.ToString() + TEXT("Assets(Id INTEGER PRIMARY KEY, Name, Class, Path TEXT UNIQUE");
+		FString InsertStatement = TEXT("INSERT INTO ") + Class.Name.ToString() + TEXT("Assets(Name, Class, Path");
+		FString ValuesStatement = TEXT("Values(?1, ?2, ?3");
+		int32 CurrentIndex = 4;
+
+		for (FName TagName : Class.Tags)
+		{
+			FString& SanitizedTagName = Class.SanitizedTagNames[TagName];
+
+			CreateTable += TEXT(", ") + SanitizedTagName;
+			InsertStatement += TEXT(", ") + SanitizedTagName;
+			ValuesStatement += FString::Printf(TEXT(", ?%d"), CurrentIndex);
+
+			Class.TagBindingIndices.Add(TagName, CurrentIndex);
+			CurrentIndex++;
+		}
+		CreateTable += TEXT(");");
+		InsertStatement += TEXT(") ") + ValuesStatement + TEXT(");");
+
+		if (OutputDatabase.Execute(*CreateTable) == false)
+		{
+			UE_LOG(LogAssetRegistryExport, Error, TEXT("Table Create Error: %s"), *OutputDatabase.GetLastError());
+			OutputDatabase.Close();
+			return 1;
+		}
+
+		if (Class.InsertStatement.Create(OutputDatabase, *InsertStatement) == false)
+		{
+			UE_LOG(LogAssetRegistryExport, Error, TEXT("Insert Statement Create Error: %s"), *OutputDatabase.GetLastError());
+			OutputDatabase.Close();
+			return 1;
+		}
+	}
+
+	// Create catchall table and insert statements
+	if (OutputDatabase.Execute(TEXT("CREATE TABLE IF NOT EXISTS Assets(Id INTEGER PRIMARY KEY, Name, Class, Path TEXT UNIQUE);")) == false)
+	{
+		UE_LOG(LogAssetRegistryExport, Error, TEXT("Error: %s"), *OutputDatabase.GetLastError());
+		OutputDatabase.Close();
+		return 1;
+	}
+
+	FSQLitePreparedStatement GlobalInsertStatement(OutputDatabase, TEXT("INSERT INTO Assets(Name, Class, Path) Values(?1, ?2, ?3);"));
+	if (GlobalInsertStatement.IsValid() == false)
+	{
+		UE_LOG(LogAssetRegistryExport, Error, TEXT("Error: %s"), *OutputDatabase.GetLastError());
+		OutputDatabase.Close();
+		return 1;
+	}
+
+	FSQLitePreparedStatement BeginTransactionStatement(OutputDatabase, TEXT("BEGIN TRANSACTION;"));
+	FSQLitePreparedStatement CommitTransactionStatement(OutputDatabase, TEXT("COMMIT TRANSACTION;"));
+
+	UE_LOG(LogAssetRegistryExport, Display, TEXT("Adding %d assets to the database..."), AssetCount);
+
+	double StartTime = FPlatformTime::Seconds();
+
+	// It's probably fine to add everything in one transaction, however add
+	// a limit for sanity.
+	const uint32 AssetsPerTransaction = 100000;
+	uint32 AssetsThisTransaction = 0;
+
+	BeginTransactionStatement.Execute();
+
+	// Actually add the assets.
+	AssetRegistry.EnumerateAllAssets(TSet<FName>(), [&](const FAssetData& AssetData)
+	{
+		FString AssetObjectPath = AssetData.ObjectPath.ToString();
+
+		// Add to the 'Assets' table.
+		{
+			GlobalInsertStatement.ClearBindings();
+
+			GlobalInsertStatement.SetBindingValueByIndex(1, AssetData.AssetName.ToString());
+			GlobalInsertStatement.SetBindingValueByIndex(2, AssetData.AssetClass.ToString());
+			GlobalInsertStatement.SetBindingValueByIndex(3, AssetObjectPath);
+
+			if (GlobalInsertStatement.Execute() == false)
+			{
+				UE_LOG(LogAssetRegistryExport, Error, TEXT("Insertion Error: %s"), *OutputDatabase.GetLastError());
+			}
+		}
+		
+		// Add to out class's table.
+		{
+			FClassInfo& Class = ClassInfos[AssetData.AssetClass];
+
+			Class.InsertStatement.ClearBindings();
+
+			Class.InsertStatement.SetBindingValueByIndex(1, AssetData.AssetName.ToString());
+			Class.InsertStatement.SetBindingValueByIndex(2, AssetData.AssetClass.ToString());
+			Class.InsertStatement.SetBindingValueByIndex(3, AssetObjectPath);
+
+			// This will leave any tag values unset (NULL) in the insertion, which is what we want.
+			AssetData.TagsAndValues.ForEach([&Class](TPair<FName, FAssetTagValueRef> Pair)
+			{
+				Class.InsertStatement.SetBindingValueByIndex(Class.TagBindingIndices[Pair.Key], Pair.Value.AsString());
+			});
+
+
+			if (Class.InsertStatement.Execute() == false)
+			{
+				UE_LOG(LogAssetRegistryExport, Error, TEXT("Class Insert Error: %s"), *OutputDatabase.GetLastError());
+			}
+		}
+
+		AssetsThisTransaction++;
+		if (AssetsThisTransaction >= AssetsPerTransaction)
+		{
+			CommitTransactionStatement.Execute();
+			BeginTransactionStatement.Execute();
+			AssetsThisTransaction = 0;
+		}
+		return true;
+	});
+
+	CommitTransactionStatement.Execute();
+
+	double EndTime = FPlatformTime::Seconds();
+
+	UE_LOG(LogAssetRegistryExport, Display, TEXT("Done in %f seconds per 1000 assets, total of %f seconds."), (EndTime - StartTime) / (AssetCount / 1000.0), EndTime - StartTime);
+
+	// Save a version so we can revisit if needed.
+	OutputDatabase.SetUserVersion(1);
+	
+	// Statements have to be destroyed before the db can be closed.
+	GlobalInsertStatement.Destroy();
+	CommitTransactionStatement.Destroy();
+	BeginTransactionStatement.Destroy();
+	for (TPair<FName, FClassInfo>& Pair : ClassInfos)
+	{
+		Pair.Value.InsertStatement.Destroy();
+	}
+
+	if (OutputDatabase.Close() == false)
+	{
+		UE_LOG(LogAssetRegistryExport, Error, TEXT("Close Error: %s"), *OutputDatabase.GetLastError());
+	}
+	UE_LOG(LogAssetRegistryExport, Display, TEXT("Done."));
+	return 0;
+}
+
+
+
+class FAssetRegistryExportModule : public IModuleInterface
+{
+public:
+	FAssetRegistryExportModule() { }
+	virtual ~FAssetRegistryExportModule() {	}
+	virtual void StartupModule() override {	}
+};
+
+IMPLEMENT_MODULE(FAssetRegistryExportModule, AssetRegistryExport);
diff --git a/Engine/Plugins/Editor/AssetRegistryExport/Source/Private/AssetRegistryExportCommandlet.h b/Engine/Plugins/Editor/AssetRegistryExport/Source/Private/AssetRegistryExportCommandlet.h
new file mode 100644
--- /dev/null
+++ b/Engine/Plugins/Editor/AssetRegistryExport/Source/Private/AssetRegistryExportCommandlet.h
@@ -0,0 +1,15 @@
+
+#pragma once
+#include "Commandlets/Commandlet.h"
+#include "AssetRegistryExportCommandlet.generated.h"
+
+DECLARE_LOG_CATEGORY_EXTERN(LogAssetRegistryExport, Log, All);
+
+UCLASS()
+class UAssetRegistryExportCommandlet : public UCommandlet
+{
+	GENERATED_UCLASS_BODY()
+public:
+	virtual int32 Main(const FString& CmdLineParams) override;
+};
+
