Fix for missing error returns.

diff --git a/Engine/Plugins/Editor/AssetRegistryExport/Source/Private/AssetRegistryExportCommandlet.cpp b/Engine/Plugins/Editor/AssetRegistryExport/Source/Private/AssetRegistryExportCommandlet.cpp
--- a/Engine/Plugins/Editor/AssetRegistryExport/Source/Private/AssetRegistryExportCommandlet.cpp
+++ b/Engine/Plugins/Editor/AssetRegistryExport/Source/Private/AssetRegistryExportCommandlet.cpp
@@ -76,11 +76,13 @@ int32 UAssetRegistryExportCommandlet::Main(const FString& CmdLineParams)
 			if (AssetRegistry.Serialize(SerializedAssetData, Options) == false)
 			{
 				UE_LOG(LogAssetRegistryExport, Error, TEXT("Failed to serialize asset registry %s"), *AssetRegistryFileName);
+				return 1;
 			}
 		}
 		else
 		{
 			UE_LOG(LogAssetRegistryExport, Error, TEXT("Failed to read asset registry %s"), *AssetRegistryFileName);
+			return 1;
 		}
 	}
 	UE_LOG(LogAssetRegistryExport, Display, TEXT("Opened asset registry: %s"), *AssetRegistryFileName);
