Asset registry export now adds the compressed size to the entire Assets table so you don't have to hunt down the asset type to find the size.

diff --git a/Engine/Plugins/Editor/AssetRegistryExport/Source/Private/AssetRegistryExportCommandlet.cpp b/Engine/Plugins/Editor/AssetRegistryExport/Source/Private/AssetRegistryExportCommandlet.cpp
--- a/Engine/Plugins/Editor/AssetRegistryExport/Source/Private/AssetRegistryExportCommandlet.cpp
+++ b/Engine/Plugins/Editor/AssetRegistryExport/Source/Private/AssetRegistryExportCommandlet.cpp
@@ -34,6 +34,9 @@ int32 UAssetRegistryExportCommandlet::Main(const FString& CmdLineParams)
 		UE_LOG(LogAssetRegistryExport, Display, TEXT("table with the corresponding tags for the class as columns. Assets are added to their"));
 		UE_LOG(LogAssetRegistryExport, Display, TEXT("class's table, as well as a global Assets table."));
 		UE_LOG(LogAssetRegistryExport, Display, TEXT(""));
+		UE_LOG(LogAssetRegistryExport, Display, TEXT("If the CompressedSize DB columns for all assets are NULL then the registry didn't have metadata "));
+		UE_LOG(LogAssetRegistryExport, Display, TEXT("written back (ProjectsSettings/Packaging/WriteBackMetadataToAssetRegistry)"));
+		UE_LOG(LogAssetRegistryExport, Display, TEXT(""));
 		UE_LOG(LogAssetRegistryExport, Display, TEXT("Parameters:"));
 		UE_LOG(LogAssetRegistryExport, Display, TEXT(""));
 		UE_LOG(LogAssetRegistryExport, Display, TEXT("    -Output=<path/to/file>              Output Sqlite DB file (Required)"));
@@ -196,14 +199,14 @@ int32 UAssetRegistryExportCommandlet::Main(const FString& CmdLineParams)
 	}
 
 	// Create catchall table and insert statements
-	if (OutputDatabase.Execute(TEXT("CREATE TABLE IF NOT EXISTS Assets(Id INTEGER PRIMARY KEY, Name, Class, Path TEXT UNIQUE);")) == false)
+	if (OutputDatabase.Execute(TEXT("CREATE TABLE IF NOT EXISTS Assets(Id INTEGER PRIMARY KEY, Name, Class, Path TEXT UNIQUE, Stage_ChunkCompressedSize INTEGER);")) == false)
 	{
 		UE_LOG(LogAssetRegistryExport, Error, TEXT("Error: %s"), *OutputDatabase.GetLastError());
 		OutputDatabase.Close();
 		return 1;
 	}
 
-	FSQLitePreparedStatement GlobalInsertStatement(OutputDatabase, TEXT("INSERT INTO Assets(Name, Class, Path) Values(?1, ?2, ?3);"));
+	FSQLitePreparedStatement GlobalInsertStatement(OutputDatabase, TEXT("INSERT INTO Assets(Name, Class, Path, Stage_ChunkCompressedSize) Values(?1, ?2, ?3, ?4);"));
 	if (GlobalInsertStatement.IsValid() == false)
 	{
 		UE_LOG(LogAssetRegistryExport, Error, TEXT("Error: %s"), *OutputDatabase.GetLastError());
@@ -224,8 +227,9 @@ int32 UAssetRegistryExportCommandlet::Main(const FString& CmdLineParams)
 	uint32 AssetsThisTransaction = 0;
 
 	BeginTransactionStatement.Execute();
-
+	
 	// Actually add the assets.
+	bool bGotCompressedSize = false;
 	AssetRegistry.EnumerateAllAssets(TSet<FName>(), [&](const FAssetData& AssetData)
 	{
 		FString AssetObjectPath = AssetData.ObjectPath.ToString();
@@ -238,13 +242,22 @@ int32 UAssetRegistryExportCommandlet::Main(const FString& CmdLineParams)
 			GlobalInsertStatement.SetBindingValueByIndex(2, AssetData.AssetClass.ToString());
 			GlobalInsertStatement.SetBindingValueByIndex(3, AssetObjectPath);
 
+			// If this assets has size information then add it to the global assets list as well.
+			FString CompressedSize;
+			if (AssetData.GetTagValue("Stage_ChunkCompressedSize", CompressedSize))
+			{
+				bGotCompressedSize = true;
+				GlobalInsertStatement.SetBindingValueByIndex(4, CompressedSize);
+			} 
+
 			if (GlobalInsertStatement.Execute() == false)
 			{
 				UE_LOG(LogAssetRegistryExport, Error, TEXT("Insertion Error: %s"), *OutputDatabase.GetLastError());
 			}
 		}
 		
-		// Add to out class's table.
+
+		// Add to our class's table.
 		{
 			FClassInfo& Class = ClassInfos[AssetData.AssetClass];
 
@@ -255,12 +268,11 @@ int32 UAssetRegistryExportCommandlet::Main(const FString& CmdLineParams)
 			Class.InsertStatement.SetBindingValueByIndex(3, AssetObjectPath);
 
 			// This will leave any tag values unset (NULL) in the insertion, which is what we want.
-			AssetData.TagsAndValues.ForEach([&Class](TPair<FName, FAssetTagValueRef> Pair)
+			AssetData.EnumerateTags([&Class, &AssetData](TPair<FName, FAssetTagValueRef> Pair)
 			{
 				Class.InsertStatement.SetBindingValueByIndex(Class.TagBindingIndices[Pair.Key], Pair.Value.AsString());
 			});
 
-
 			if (Class.InsertStatement.Execute() == false)
 			{
 				UE_LOG(LogAssetRegistryExport, Error, TEXT("Class Insert Error: %s"), *OutputDatabase.GetLastError());
@@ -283,6 +295,13 @@ int32 UAssetRegistryExportCommandlet::Main(const FString& CmdLineParams)
 
 	UE_LOG(LogAssetRegistryExport, Display, TEXT("Done in %f seconds per 1000 assets, total of %f seconds."), (EndTime - StartTime) / (AssetCount / 1000.0), EndTime - StartTime);
 
+	if (bGotCompressedSize == false)
+	{
+		UE_LOG(LogAssetRegistryExport, Display, TEXT("Asset registry didn't have size metadata written back, CompressedSize DB column will be NULL."));
+		UE_LOG(LogAssetRegistryExport, Display, TEXT("Metadata can be written back via ProjectSettings/Packaging/WriteBackMetadataToAssetRegistry"));
+		UE_LOG(LogAssetRegistryExport, Display, TEXT("or after staging with iostore -AssetRegistryWriteback."));
+	}
+
 	// Save a version so we can revisit if needed.
 	OutputDatabase.SetUserVersion(1);
 	
