Class names now seem to include /Script/ which is invalid in sqlite.

diff --git a/Engine/Plugins/Editor/AssetRegistryExport/Source/Private/AssetRegistryExportCommandlet.cpp b/Engine/Plugins/Editor/AssetRegistryExport/Source/Private/AssetRegistryExportCommandlet.cpp
--- a/Engine/Plugins/Editor/AssetRegistryExport/Source/Private/AssetRegistryExportCommandlet.cpp
+++ b/Engine/Plugins/Editor/AssetRegistryExport/Source/Private/AssetRegistryExportCommandlet.cpp
@@ -109,7 +109,7 @@ int32 UAssetRegistryExportCommandlet::Main(const FString& CmdLineParams)
 
 	struct FClassInfo
 	{
-		FTopLevelAssetPath Name;
+		FString Name;
 		TSet<FName> Tags;
 
 		// TagName -> Name that has been sanitized for use as a column name.
@@ -129,7 +129,17 @@ int32 UAssetRegistryExportCommandlet::Main(const FString& CmdLineParams)
 	AssetRegistry.EnumerateAllAssets(TSet<FName>(), [&](const FAssetData& AssetData)
 	{
 		FClassInfo& ClassInfo = ClassInfos.FindOrAdd(AssetData.AssetClassPath);
-		ClassInfo.Name = AssetData.AssetClassPath;
+		ClassInfo.Name = AssetData.AssetClassPath.ToString();
+		{
+			auto& CharArray = ClassInfo.Name.GetCharArray();
+			for (int i = 0; i < CharArray.Num() - 1; i++)
+			{
+				if (TChar<TCHAR>::IsAlnum(CharArray[i]) == false)
+				{
+					CharArray[i] = TEXT('_');
+				}
+			}
+		}
 
 		AssetData.TagsAndValues.ForEach([&ClassInfo](TPair<FName, FAssetTagValueRef> Pair)
 		{
@@ -164,8 +174,8 @@ int32 UAssetRegistryExportCommandlet::Main(const FString& CmdLineParams)
 	{
 		FClassInfo& Class = Pair.Value;
 
-		FString CreateTable = TEXT("CREATE TABLE IF NOT EXISTS ") + Class.Name.ToString() + TEXT("Assets(Id INTEGER PRIMARY KEY, Name, Class, Path TEXT UNIQUE");
-		FString InsertStatement = TEXT("INSERT INTO ") + Class.Name.ToString() + TEXT("Assets(Name, Class, Path");
+		FString CreateTable = TEXT("CREATE TABLE IF NOT EXISTS ") + Class.Name + TEXT("Assets(Id INTEGER PRIMARY KEY, Name, Class, Path TEXT UNIQUE");
+		FString InsertStatement = TEXT("INSERT INTO ") + Class.Name + TEXT("Assets(Name, Class, Path");
 		FString ValuesStatement = TEXT("Values(?1, ?2, ?3");
 		int32 CurrentIndex = 4;
 
