Update the asset registry exporter to output to CSV to facilitate import in to common database frontends. Additionally, support outputting dependency information based on a set of primary asset type "roots" for help in analyzing the source of build size.

diff --git a/Engine/Plugins/Editor/AssetRegistryExport/Source/Private/AssetRegistryExportCommandlet.cpp b/Engine/Plugins/Editor/AssetRegistryExport/Source/Private/AssetRegistryExportCommandlet.cpp
--- a/Engine/Plugins/Editor/AssetRegistryExport/Source/Private/AssetRegistryExportCommandlet.cpp
+++ b/Engine/Plugins/Editor/AssetRegistryExport/Source/Private/AssetRegistryExportCommandlet.cpp
@@ -3,9 +3,15 @@
 #include "AssetRegistryExportCommandlet.h"
 
 #include "AssetRegistry/AssetRegistryState.h"
+#include "Containers/Queue.h"
+#include "HAL/FileManager.h"
+#include "Misc/App.h"
 #include "Misc/FileHelper.h"
+#include "Misc/Paths.h"
+#include "Misc/StringBuilder.h"
 #include "Modules/ModuleManager.h"
 #include "Serialization/ArrayReader.h"
+#include "Serialization/MemoryWriter.h"
 #include "SQLiteDatabase.h"
 #include "SQLitePreparedStatement.h"
 
@@ -17,18 +23,31 @@ UAssetRegistryExportCommandlet::UAssetRegistryExportCommandlet(const FObjectInit
 {
 }
 
+// Helper to comma separate numbers for readability.
+static FString NumberString(uint64 N) { return FText::AsNumber(N).ToString(); }
+
+
 int32 UAssetRegistryExportCommandlet::Main(const FString& CmdLineParams)
 {
 	FString OutputDatabaseName;
 	FString AssetRegistryFileName;
+	FString CSVFileName;
+	FString PrimaryAssetTypes;
+
+	FParse::Value(*CmdLineParams, TEXT("ListDependencies="), PrimaryAssetTypes);
+	bool bIncludeSharedWithDependencies = FParse::Param(*CmdLineParams, TEXT("IncludeSharedWithDependencies"));
+
+	// We must have either a sqlite or a csv output
+	bool bHaveOutput = FParse::Value(*CmdLineParams, TEXT("Output="), OutputDatabaseName) ||
+		FParse::Value(*CmdLineParams, TEXT("CSV="), CSVFileName);
 
-	if (FParse::Value(*CmdLineParams, TEXT("Output="), OutputDatabaseName) == false ||
+	if (bHaveOutput == false ||
 		FParse::Value(*CmdLineParams, TEXT("Input="), AssetRegistryFileName) == false)
 	{
 		UE_LOG(LogAssetRegistryExport, Display, TEXT(""));
 		UE_LOG(LogAssetRegistryExport, Display, TEXT("Asset Registry Exporter"));
 		UE_LOG(LogAssetRegistryExport, Display, TEXT(""));
-		UE_LOG(LogAssetRegistryExport, Display, TEXT("Converts a given asset registry to a sqlite database. Each AssetClass gets its own"));
+		UE_LOG(LogAssetRegistryExport, Display, TEXT("Converts a given asset registry to a sqlite/csv database. Each AssetClass gets its own"));
 		UE_LOG(LogAssetRegistryExport, Display, TEXT("table with the corresponding tags for the class as columns. Assets are added to their"));
 		UE_LOG(LogAssetRegistryExport, Display, TEXT("class's table, as well as a global Assets table."));
 		UE_LOG(LogAssetRegistryExport, Display, TEXT(""));
@@ -37,36 +56,49 @@ int32 UAssetRegistryExportCommandlet::Main(const FString& CmdLineParams)
 		UE_LOG(LogAssetRegistryExport, Display, TEXT(""));
 		UE_LOG(LogAssetRegistryExport, Display, TEXT("Parameters:"));
 		UE_LOG(LogAssetRegistryExport, Display, TEXT(""));
-		UE_LOG(LogAssetRegistryExport, Display, TEXT("    -Output=<path/to/file>              Output Sqlite DB file (Required)"));
+		UE_LOG(LogAssetRegistryExport, Display, TEXT("    -Output=<path/to/file>              Output Sqlite DB file (Required Or CSV)"));
+		UE_LOG(LogAssetRegistryExport, Display, TEXT("    -CSV=<path>                         Directory for CSV output files (Required Or Sqlite) MAKES A TON OF FILES."));
 		UE_LOG(LogAssetRegistryExport, Display, TEXT("    -Input=<path/to/file>               Input Asset Registry file (Required)"));
-		UE_LOG(LogAssetRegistryExport, Display, TEXT("    -Overwrite                          Allow overwrite of existing db file (Optional)."));
-		UE_LOG(LogAssetRegistryExport, Display, TEXT("                                        By default existing output files is a failure."));
+		UE_LOG(LogAssetRegistryExport, Display, TEXT("    -ListDependencies=typea,typeb       Comma separated list of primary asset types to output dependency"));
+		UE_LOG(LogAssetRegistryExport, Display, TEXT("                                        information. Requires asset registry metadata and to be running as"));
+		UE_LOG(LogAssetRegistryExport, Display, TEXT("                                        a commandlet for the project that generated the asset registry."));
+		UE_LOG(LogAssetRegistryExport, Display, TEXT("    -IncludeSharedWithDependencies      When listing dependencies, this will also add shared dependencies"));
+		UE_LOG(LogAssetRegistryExport, Display, TEXT("                                        to a \"shared\" root."));
 		return 1;
 	}
 
-	if (IFileManager::Get().FileExists(*OutputDatabaseName))
+	FSQLiteDatabase OutputDatabase;
+
+	bool bCSV = CSVFileName.Len() != 0;
+
+	if (bCSV == false)
 	{
-		if (FParse::Param(*CmdLineParams, TEXT("Overwrite")) == false)
+		if (OutputDatabase.Open(*OutputDatabaseName) == false)
 		{
-			UE_LOG(LogAssetRegistryExport, Error, TEXT("Output database exists, delete or use -Overwrite. (%s)"), *OutputDatabaseName);
+			UE_LOG(LogAssetRegistryExport, Error, TEXT("Failed to open sqlite database %s"), *OutputDatabaseName);
 			return 1;
 		}
-
-		UE_LOG(LogAssetRegistryExport, Display, TEXT("Output database exists, deleting."));
-		if (IFileManager::Get().Delete(*OutputDatabaseName) == false)
+		UE_LOG(LogAssetRegistryExport, Display, TEXT("Creating Sqlite DB: %s"), *OutputDatabaseName);
+	}
+	else
+	{
+		FPaths::MakePlatformFilename(CSVFileName);
+		if (CSVFileName.EndsWith(FPlatformMisc::GetDefaultPathSeparator()) == false)
 		{
-			UE_LOG(LogAssetRegistryExport, Error, TEXT("Failed to delete output database %s."), *OutputDatabaseName);
-			return 1;
+			CSVFileName += FPlatformMisc::GetDefaultPathSeparator();
 		}
 	}
 
-	FSQLiteDatabase OutputDatabase;
-	if (OutputDatabase.Open(*OutputDatabaseName) == false)
+	TSet<FName> DependencyListAssetTypeNames;
+	if (PrimaryAssetTypes.Len())
 	{
-		UE_LOG(LogAssetRegistryExport, Error, TEXT("Failed to open sqlite database %s"), *OutputDatabaseName);
-		return 1;
+		TArray<FString> DependencyListAssetTypes;
+		PrimaryAssetTypes.ParseIntoArray(DependencyListAssetTypes, TEXT(","), true);
+		for (FString& AssetType : DependencyListAssetTypes)
+		{
+			DependencyListAssetTypeNames.Add(FName(AssetType));
+		}
 	}
-	UE_LOG(LogAssetRegistryExport, Display, TEXT("Creating Sqlite DB: %s"), *OutputDatabaseName);
 
 	FAssetRegistryState AssetRegistry;
 	{
@@ -88,8 +120,23 @@ int32 UAssetRegistryExportCommandlet::Main(const FString& CmdLineParams)
 	}
 	UE_LOG(LogAssetRegistryExport, Display, TEXT("Opened asset registry: %s"), *AssetRegistryFileName);
 
+	if (DependencyListAssetTypeNames.Num())
+	{
+		// We need the project loaded to do the dependency listing because we look at the class hierarchy
+		// which currently isn't stored in the asset registry.
+		FString CheckForProjectName = FString("/Script/") + FApp::GetProjectName();
+		TArray<FAssetDependency> CheckDeps;
+		AssetRegistry.GetReferencers(FAssetIdentifier(*CheckForProjectName), CheckDeps);
+		if (CheckDeps.Num() == 0)
+		{
+			UE_LOG(LogAssetRegistryExport, Error, TEXT("Can't list dependencies without running as a commandlet for the same project that generated the DevelopmentAssetRegistry"));
+			return 1;
+		}
+	}
+
 	// Set up some database properties for speed - we don't really care about ACID in this case
 	// since we're just dumping stuff and closing the database - we can always start over.
+	if (bCSV == false)
 	{
 		// don't bother waiting for disk to flush (let the OS handle it)
 		OutputDatabase.Execute(TEXT("PRAGMA synchronous=OFF;"));
@@ -117,6 +164,8 @@ int32 UAssetRegistryExportCommandlet::Main(const FString& CmdLineParams)
 		TMap<FName, int32> TagBindingIndices;
 
 		FSQLitePreparedStatement InsertStatement;
+
+		TArray64<uint8> CSVMemory;
 	};
 
 	TMap<FTopLevelAssetPath, FClassInfo> ClassInfos;
@@ -165,6 +214,419 @@ int32 UAssetRegistryExportCommandlet::Main(const FString& CmdLineParams)
 		return true;
 	});
 
+	if (DependencyListAssetTypeNames.Num())
+	{
+		// function to try to find a base class that makes sense for the given asset.
+		auto GetRootClassToPrint = [](FTopLevelAssetPath AssetClassPath)
+		{
+			UClass* AssetClass = FindObject<UClass>(AssetClassPath, true);
+			if (AssetClass == nullptr)
+			{
+				return FString("Unknown");
+			}
+
+			UClass* OriginalClass = AssetClass;
+			FString ClassName;
+			while (UClass* Super = AssetClass->GetSuperClass())
+			{				
+				Super->GetName(ClassName);
+
+				// Stop when we get to something too generic.
+				if (ClassName == "Object" || ClassName == "BlueprintCore" || ClassName == "DataAsset" || ClassName == "StreamableRenderAsset" || ClassName == "SoundBase")
+				{
+					break;
+				}
+				AssetClass = Super;
+			}
+
+			if (AssetClass)
+			{
+				return AssetClass->GetName();
+			}
+			else
+			{
+				return OriginalClass->GetName();
+			}
+		};
+
+		//
+		// Create a table to hold each game feature data dependency list.
+		//
+
+		// Find all the assets that are of the requested primary asset types. These we'll use as roots for the 
+		// dependency listing.
+		TArray<const FAssetData*> PrimaryAssets;
+		const TArray<const FAssetData*>& AssetsWithPrimaryAssetType = AssetRegistry.GetAssetsByTagName(FPrimaryAssetId::PrimaryAssetTypeTag);
+		for (const FAssetData* GFD : AssetsWithPrimaryAssetType)
+		{
+			FName PrimaryAssetType;
+			if (GFD->GetTagValue(FPrimaryAssetId::PrimaryAssetTypeTag, PrimaryAssetType) &&
+				DependencyListAssetTypeNames.Contains(PrimaryAssetType))
+			{
+				PrimaryAssets.Add(GFD);
+			}
+		}
+
+		
+		UE_LOG(LogAssetRegistryExport, Display, TEXT("ListDependencies matching PrimaryAssets discovered %s"), *NumberString(PrimaryAssets.Num()));
+
+		// Track all assets with size specified so we can know who isn't assigned.
+		TSet<const FAssetData*> RemainingAssetsWithSize;
+		uint64 TotalCompressedBytes = 0;
+		{
+			const TArray<const FAssetData*>& AssetsWithSize = AssetRegistry.GetAssetsByTagName("Stage_ChunkCompressedSize");
+			for (const FAssetData* AD : AssetsWithSize)
+			{
+				uint64 CompressedSize = 0;
+				if (AD->GetTagValue("Stage_ChunkCompressedSize", CompressedSize))
+				{
+					TotalCompressedBytes += CompressedSize;
+					RemainingAssetsWithSize.Add(AD);
+				}
+			}
+		}
+
+		const TMap<FName, const FAssetPackageData*>& AllPackages = AssetRegistry.GetAssetPackageDataMap();
+
+		TSet<FName> UnionedTotalDependencies;
+		TSet<FName> UnionUniqueDependencies;
+
+		struct FPrimaryAssetInfo
+		{
+			TSet<FName> TotalPackageDependencies;
+			TSet<FName> UniquePackageDependencies;
+
+			uint64 SizeOfAllDependencies = 0;
+			uint64 SizeOfUniqueDependencies = 0;
+		};
+	
+		TMap<FName, FPrimaryAssetInfo> PrimaryAssetInfos;
+
+		uint64 TotalSizesOfUniqueDependencies = 0;
+		uint64 UnassignedSize = 0;
+		uint64 SharedSize = 0;
+
+		{
+			TArray64<uint8> CSV;
+			FMemoryWriter64 CSVWriter(CSV);
+
+			{
+				TUtf8StringBuilder<32> Header;
+				Header << "RootPrimaryAsset,Name,Class,RootClass,Stage_ChunkCompressedSize";
+				Header << "\r\n";
+				CSVWriter.Serialize(Header.GetData(), Header.Len());
+			}
+
+			for (const FAssetData* PrimaryAsset : PrimaryAssets)
+			{
+				FPrimaryAssetInfo& Info = PrimaryAssetInfos.Add(PrimaryAsset->PackageName);
+				Info.TotalPackageDependencies.Add(PrimaryAsset->PackageName);
+
+				// Get recursive dependencies
+				{
+					TQueue<FName> PackageQueue;
+					PackageQueue.Enqueue(PrimaryAsset->PackageName);
+
+					FName PackageName;
+					while (PackageQueue.Dequeue(PackageName))
+					{
+						TArrayView<const FAssetData* const> AssetDataList = AssetRegistry.GetAssetsByPackageName(PackageName);
+
+						TArray<FAssetIdentifier> PackageDependencies;
+						AssetRegistry.GetDependencies(FAssetIdentifier(PackageName), PackageDependencies);
+
+						for (FAssetIdentifier PackageDependency : PackageDependencies)
+						{
+							if (PackageDependency.IsPackage())
+							{
+								bool bAlreadyExisted = false;
+								Info.TotalPackageDependencies.Add(PackageDependency.PackageName, &bAlreadyExisted);
+
+								if (bAlreadyExisted == false)
+								{
+									PackageQueue.Enqueue(PackageDependency.PackageName);
+								}
+							}
+						}
+					}
+				}
+
+				
+
+				// Write out the total dependencies and sum the size of everything. (this is the slow part)
+				Info.SizeOfAllDependencies = 0;
+				TUtf8StringBuilder<2048> Line;
+				for (FName& PackageDependency : Info.TotalPackageDependencies)
+				{
+					Line.Reset();
+					Line << PrimaryAsset->GetObjectPathString() << ",";
+					Line << PackageDependency.ToString() << ",";
+
+					const FAssetData* UseAsClassData = UE::AssetRegistry::GetMostImportantAsset(AssetRegistry.GetAssetsByPackageName(PackageDependency), UE::AssetRegistry::EGetMostImportantAssetFlags::IgnoreSkipClasses);
+					if (UseAsClassData)
+					{
+						Line << UseAsClassData->AssetClassPath.ToString() << ",";
+						Line << GetRootClassToPrint(UseAsClassData->AssetClassPath) << ",";
+
+						uint64 CompressedSize = 0;
+						if (UseAsClassData->GetTagValue("Stage_ChunkCompressedSize", CompressedSize))
+						{
+							Info.SizeOfAllDependencies += CompressedSize;
+						}
+						Line << CompressedSize << ",";
+					}
+					else
+					{
+						Line << "Unknown,,";
+					}
+			
+					Line << "\r\n";
+					CSVWriter.Serialize(Line.GetData(), Line.Len());
+
+					UnionedTotalDependencies.Add(PackageDependency);
+				}
+
+				
+
+			} // end for each root
+
+
+			TMap<const FAssetData*, FName> AccountedFor;
+
+			// Now figure out what's unique to each root. Iterate over all dependencies for
+			// all roots and mark when we see it more than once.
+			TMap<FName, bool> IsDependencyUnique;
+			for (TPair<FName, FPrimaryAssetInfo>& Pair : PrimaryAssetInfos)
+			{
+				for (FName PackageDependency : Pair.Value.TotalPackageDependencies)
+				{
+					bool* bIsUnique = IsDependencyUnique.Find(PackageDependency);
+					if (bIsUnique)
+					{
+						*bIsUnique = false; // we've seen it before so it can't be unique.
+					}
+					else
+					{
+						IsDependencyUnique.Add(PackageDependency, true);
+					}
+				}
+			}
+
+			// Now that we have uniqueness, go back and copy the entries for each root that
+			// are unique and sum the size.
+			for (TPair<FName, FPrimaryAssetInfo>& Pair : PrimaryAssetInfos)
+			{
+				for (FName PackageDependency : Pair.Value.TotalPackageDependencies)
+				{
+					bool bIsUnique = IsDependencyUnique.FindChecked(PackageDependency);
+					if (bIsUnique)
+					{
+						Pair.Value.UniquePackageDependencies.Add(PackageDependency);
+						UnionUniqueDependencies.Add(PackageDependency);
+
+						const FAssetData* UseAsClassData = UE::AssetRegistry::GetMostImportantAsset(AssetRegistry.GetAssetsByPackageName(PackageDependency), UE::AssetRegistry::EGetMostImportantAssetFlags::IgnoreSkipClasses);
+						if (UseAsClassData)
+						{
+							uint64 CompressedSize = 0;
+							if (UseAsClassData->GetTagValue("Stage_ChunkCompressedSize", CompressedSize))
+							{
+								check(RemainingAssetsWithSize.Contains(UseAsClassData));
+								RemainingAssetsWithSize.Remove(UseAsClassData);
+
+								if (AccountedFor.Contains(UseAsClassData))
+								{
+									UE_LOG(LogAssetRegistryExport, Warning, TEXT("Already seen %s in %s, we are %s"), *PackageDependency.ToString(), *AccountedFor[UseAsClassData].ToString(), *Pair.Key.ToString());
+								}
+								AccountedFor.Add(UseAsClassData, Pair.Key);
+
+								Pair.Value.SizeOfUniqueDependencies += CompressedSize;
+							}
+						}
+					}
+				}
+
+				TotalSizesOfUniqueDependencies += Pair.Value.SizeOfUniqueDependencies;
+			}
+
+			for (const TPair<FName, const FAssetPackageData*>& Package : AllPackages)
+			{			
+				if (UnionedTotalDependencies.Contains(Package.Key) == false)
+				{
+					// If the package isn't in _any_ dependency, it's unassigned.
+					TUtf8StringBuilder<128> Line;
+					Line << "Unassigned,";
+					Line << Package.Key.ToString() << ",";
+
+					const FAssetData* UseAsClassData = UE::AssetRegistry::GetMostImportantAsset(AssetRegistry.GetAssetsByPackageName(Package.Key), UE::AssetRegistry::EGetMostImportantAssetFlags::IgnoreSkipClasses);
+					if (UseAsClassData)
+					{
+						Line << UseAsClassData->AssetClassPath.ToString() << ",";
+						Line << GetRootClassToPrint(UseAsClassData->AssetClassPath) << ",";
+
+						uint64 CompressedSize = 0;
+						if (UseAsClassData->GetTagValue("Stage_ChunkCompressedSize", CompressedSize))
+						{
+							check(RemainingAssetsWithSize.Contains(UseAsClassData));
+							RemainingAssetsWithSize.Remove(UseAsClassData);
+
+							if (AccountedFor.Contains(UseAsClassData))
+							{
+								UE_LOG(LogAssetRegistryExport, Warning, TEXT("Already seen %s in %s, we are %s"), *Package.Key.ToString(), *AccountedFor[UseAsClassData].ToString(), *Package.Key.ToString());
+							}
+							AccountedFor.Add(UseAsClassData, Package.Key);
+
+							UnassignedSize += CompressedSize;
+							RemainingAssetsWithSize.Remove(UseAsClassData);
+							Line << CompressedSize;
+						}					
+					}
+					else
+					{
+						Line << "Unknown,";
+					}
+
+					Line << "\r\n";
+					CSVWriter.Serialize(Line.GetData(), Line.Len());
+					continue;
+				}
+
+				if (UnionUniqueDependencies.Contains(Package.Key) == false)
+				{
+					// If the package IS _any_ dependency, but not in any _unique_ dependency, it's a shared
+					// dependency.
+					TUtf8StringBuilder<128> Line;
+					Line << "Shared,";
+					Line << Package.Key.ToString() << ",";
+
+					const FAssetData* UseAsClassData = UE::AssetRegistry::GetMostImportantAsset(AssetRegistry.GetAssetsByPackageName(Package.Key), UE::AssetRegistry::EGetMostImportantAssetFlags::IgnoreSkipClasses);
+					if (UseAsClassData)
+					{
+						Line << UseAsClassData->AssetClassPath.ToString() << ",";
+						Line << GetRootClassToPrint(UseAsClassData->AssetClassPath) << ",";
+
+						uint64 CompressedSize = 0;
+						if (UseAsClassData->GetTagValue("Stage_ChunkCompressedSize", CompressedSize))
+						{
+							check(RemainingAssetsWithSize.Contains(UseAsClassData));
+							RemainingAssetsWithSize.Remove(UseAsClassData);
+
+							if (AccountedFor.Contains(UseAsClassData))
+							{
+								UE_LOG(LogAssetRegistryExport, Warning, TEXT("Already seen %s in %s, we are %s"), *Package.Key.ToString(), *AccountedFor[UseAsClassData].ToString(), *Package.Key.ToString());
+							}
+							AccountedFor.Add(UseAsClassData, Package.Key);
+
+							SharedSize += CompressedSize;
+							RemainingAssetsWithSize.Remove(UseAsClassData);
+							Line << CompressedSize;
+						}
+					}
+					else
+					{
+						Line << "Unknown,";
+					}
+
+					Line << "\r\n";
+
+					if (bIncludeSharedWithDependencies)
+					{
+						CSVWriter.Serialize(Line.GetData(), Line.Len());
+					}
+				}
+			}
+
+			{
+				FArchive* Ar = IFileManager::Get().CreateFileWriter(*(CSVFileName + TEXT("ListedDependencies.csv")));
+				Ar->Serialize(CSV.GetData(), CSV.Num());
+				delete Ar;
+			}
+
+			// Write out a shorter CSV that just summarizes the size of the various roots.
+			{
+				TArray64<uint8> SizeCSV;
+				FMemoryWriter64 SizeCSVWriter(SizeCSV);
+
+				{
+					TUtf8StringBuilder<32> Header;
+					Header << "RootPrimaryAsset,AllDependenciesCompressedBytes,UniqueDependenciesCompressedBytes";
+					Header << "\r\n";
+					SizeCSVWriter.Serialize(Header.GetData(), Header.Len());
+				}
+
+				for (TPair<FName, FPrimaryAssetInfo>& Pair : PrimaryAssetInfos)
+				{
+					TUtf8StringBuilder<128> Line;
+
+					Line << Pair.Key.ToString() << ",";
+					Line << Pair.Value.SizeOfAllDependencies << ",";
+					Line << Pair.Value.SizeOfUniqueDependencies;
+					Line << "\r\n";
+					SizeCSVWriter.Serialize(Line.GetData(), Line.Len());
+				}
+
+				{
+					TUtf8StringBuilder<128> Line;
+
+					Line <<"Unassigned,";
+					Line << UnassignedSize << ",";
+					Line << UnassignedSize;
+					Line << "\r\n";
+					SizeCSVWriter.Serialize(Line.GetData(), Line.Len());
+				}
+
+				{
+					TUtf8StringBuilder<128> Line;
+
+					Line << "Shared,";
+					Line << SharedSize << ",";
+					Line << SharedSize;
+					Line << "\r\n";
+					SizeCSVWriter.Serialize(Line.GetData(), Line.Len());
+				}
+
+				FArchive* Ar = IFileManager::Get().CreateFileWriter(*(CSVFileName + TEXT("SizeDependencies.csv")));
+				Ar->Serialize(SizeCSV.GetData(), SizeCSV.Num());
+				delete Ar;
+			}
+		}
+
+		UE_LOG(LogAssetRegistryExport, Display, TEXT("TotalUnique: %s, Shared: %s Unassigned: %s, Total %s, Orphaned (%s assets) %s"),
+			*NumberString(TotalSizesOfUniqueDependencies),
+			*NumberString(SharedSize),
+			*NumberString(UnassignedSize), 
+			*NumberString(TotalCompressedBytes), 
+			*NumberString(RemainingAssetsWithSize.Num()),
+			*NumberString(TotalCompressedBytes - UnassignedSize - TotalSizesOfUniqueDependencies - SharedSize));
+
+		if (RemainingAssetsWithSize.Num())
+		{
+			UE_LOG(LogAssetRegistryExport, Display, TEXT("%s orphaned assets found - this is due to differences in assigning the \"most important\" asset in a package between UnrealPak and here."), *NumberString(RemainingAssetsWithSize.Num()));
+		}
+		for (const FAssetData* LD : RemainingAssetsWithSize)
+		{
+			uint64 value = 0;
+			if (LD->GetTagValue("Stage_ChunkCompressedSize", value))
+			{
+				const FAssetPackageData* PackageData = AllPackages[LD->PackageName];
+
+				UE_LOG(LogAssetRegistryExport, Warning, TEXT("...Orphaned: %s (%s bytes, in packagemap: %s"), *LD->PackageName.ToString(), *NumberString(value), *LexToString(AllPackages.Contains(LD->PackageName)));
+
+				TArrayView<const FAssetData* const> PackageAssets = AssetRegistry.GetAssetsByPackageName(LD->PackageName);
+				for (const FAssetData* AD : PackageAssets)
+				{
+					uint64 blah = 0;
+					AD->GetTagValue("Stage_ChunkCompressedSize", blah);
+
+					UE_LOG(LogAssetRegistryExport, Warning, TEXT("......Package has %s %d"), *AD->AssetName.ToString(), blah);
+
+				}
+
+				const FAssetData* MostImportant = UE::AssetRegistry::GetMostImportantAsset(PackageAssets, UE::AssetRegistry::EGetMostImportantAssetFlags::IgnoreSkipClasses);
+				UE_LOG(LogAssetRegistryExport, Warning, TEXT("......MostImportant: %s"), MostImportant ? *MostImportant->AssetName.ToString() : TEXT("null"));
+			}
+		}
+	} // end if dependency list is requested.
+
 	//
 	// Create the tables and insert statements for each class
 	//
@@ -172,59 +634,92 @@ int32 UAssetRegistryExportCommandlet::Main(const FString& CmdLineParams)
 	{
 		FClassInfo& Class = Pair.Value;
 
-		FString CreateTable = TEXT("CREATE TABLE IF NOT EXISTS ") + Class.Name + TEXT("Assets(Id INTEGER PRIMARY KEY, Name, Class, Path TEXT UNIQUE");
-		FString InsertStatement = TEXT("INSERT INTO ") + Class.Name + TEXT("Assets(Name, Class, Path");
-		FString ValuesStatement = TEXT("Values(?1, ?2, ?3");
-		int32 CurrentIndex = 4;
+		if (bCSV)
+		{
+			TUtf8StringBuilder<128> Header;
+			Header << "Name,Class,Path";
 
-		for (FName TagName : Class.Tags)
+			for (FName TagName : Class.Tags)
+			{
+				FString& SanitizedTagName = Class.SanitizedTagNames[TagName];
+				Header << "," << SanitizedTagName;
+			}
+			Header << "\r\n";
+
+			FMemoryWriter64 CSVWriter(Class.CSVMemory);
+			CSVWriter.Serialize(Header.GetData(), Header.Len());
+		}
+		else
 		{
-			FString& SanitizedTagName = Class.SanitizedTagNames[TagName];
+			FString CreateTable = TEXT("CREATE TABLE IF NOT EXISTS ") + Class.Name + TEXT("Assets(Id INTEGER PRIMARY KEY, Name, Class, Path TEXT UNIQUE");
+			FString InsertStatement = TEXT("INSERT INTO ") + Class.Name + TEXT("Assets(Name, Class, Path");
+			FString ValuesStatement = TEXT("Values(?1, ?2, ?3");
+			int32 CurrentIndex = 4;
+
+			for (FName TagName : Class.Tags)
+			{
+				FString& SanitizedTagName = Class.SanitizedTagNames[TagName];
+
+				CreateTable += TEXT(", ") + SanitizedTagName;
+				InsertStatement += TEXT(", ") + SanitizedTagName;
+				ValuesStatement += FString::Printf(TEXT(", ?%d"), CurrentIndex);
 
-			CreateTable += TEXT(", ") + SanitizedTagName;
-			InsertStatement += TEXT(", ") + SanitizedTagName;
-			ValuesStatement += FString::Printf(TEXT(", ?%d"), CurrentIndex);
+				Class.TagBindingIndices.Add(TagName, CurrentIndex);
+				CurrentIndex++;
+			}
+			CreateTable += TEXT(");");
+			InsertStatement += TEXT(") ") + ValuesStatement + TEXT(");");
 
-			Class.TagBindingIndices.Add(TagName, CurrentIndex);
-			CurrentIndex++;
+			if (OutputDatabase.Execute(*CreateTable) == false)
+			{
+				UE_LOG(LogAssetRegistryExport, Error, TEXT("Table Create Error: %s"), *OutputDatabase.GetLastError());
+				OutputDatabase.Close();
+				return 1;
+			}
+
+			if (Class.InsertStatement.Create(OutputDatabase, *InsertStatement) == false)
+			{
+				UE_LOG(LogAssetRegistryExport, Error, TEXT("Insert Statement Create Error: %s"), *OutputDatabase.GetLastError());
+				OutputDatabase.Close();
+				return 1;
+			}
 		}
-		CreateTable += TEXT(");");
-		InsertStatement += TEXT(") ") + ValuesStatement + TEXT(");");
+	}
 
-		if (OutputDatabase.Execute(*CreateTable) == false)
+	TArray64<uint8> AssetsCSV;
+	TOptional<FSQLitePreparedStatement> GlobalInsertStatement;
+	TOptional<FSQLitePreparedStatement> BeginTransactionStatement;
+	TOptional<FSQLitePreparedStatement> CommitTransactionStatement;
+	if (bCSV)
+	{
+		FMemoryWriter64 CSVWriter(AssetsCSV);
+		
+		TUtf8StringBuilder<128> Header;
+		Header << "Name,Class,Path,Stage_ChunkCompressedSize\r\n";
+		CSVWriter.Serialize(Header.GetData(), Header.Len());
+	}
+	else
+	{
+		// Create catchall table and insert statements
+		if (OutputDatabase.Execute(TEXT("CREATE TABLE IF NOT EXISTS Assets(Id INTEGER PRIMARY KEY, Name, Class, Path TEXT UNIQUE, Stage_ChunkCompressedSize INTEGER);")) == false)
 		{
-			UE_LOG(LogAssetRegistryExport, Error, TEXT("Table Create Error: %s"), *OutputDatabase.GetLastError());
+			UE_LOG(LogAssetRegistryExport, Error, TEXT("Error: %s"), *OutputDatabase.GetLastError());
 			OutputDatabase.Close();
 			return 1;
 		}
 
-		if (Class.InsertStatement.Create(OutputDatabase, *InsertStatement) == false)
+		GlobalInsertStatement.Emplace(OutputDatabase, TEXT("INSERT INTO Assets(Name, Class, Path, Stage_ChunkCompressedSize) Values(?1, ?2, ?3, ?4);"));
+		if (GlobalInsertStatement.GetPtrOrNull()->IsValid() == false)
 		{
-			UE_LOG(LogAssetRegistryExport, Error, TEXT("Insert Statement Create Error: %s"), *OutputDatabase.GetLastError());
+			UE_LOG(LogAssetRegistryExport, Error, TEXT("Error: %s"), *OutputDatabase.GetLastError());
 			OutputDatabase.Close();
 			return 1;
 		}
-	}
 
-	// Create catchall table and insert statements
-	if (OutputDatabase.Execute(TEXT("CREATE TABLE IF NOT EXISTS Assets(Id INTEGER PRIMARY KEY, Name, Class, Path TEXT UNIQUE, Stage_ChunkCompressedSize INTEGER);")) == false)
-	{
-		UE_LOG(LogAssetRegistryExport, Error, TEXT("Error: %s"), *OutputDatabase.GetLastError());
-		OutputDatabase.Close();
-		return 1;
+		BeginTransactionStatement.Emplace(OutputDatabase, TEXT("BEGIN TRANSACTION;"));
+		CommitTransactionStatement.Emplace(OutputDatabase, TEXT("COMMIT TRANSACTION;"));
 	}
 
-	FSQLitePreparedStatement GlobalInsertStatement(OutputDatabase, TEXT("INSERT INTO Assets(Name, Class, Path, Stage_ChunkCompressedSize) Values(?1, ?2, ?3, ?4);"));
-	if (GlobalInsertStatement.IsValid() == false)
-	{
-		UE_LOG(LogAssetRegistryExport, Error, TEXT("Error: %s"), *OutputDatabase.GetLastError());
-		OutputDatabase.Close();
-		return 1;
-	}
-
-	FSQLitePreparedStatement BeginTransactionStatement(OutputDatabase, TEXT("BEGIN TRANSACTION;"));
-	FSQLitePreparedStatement CommitTransactionStatement(OutputDatabase, TEXT("COMMIT TRANSACTION;"));
-
 	UE_LOG(LogAssetRegistryExport, Display, TEXT("Adding %d assets to the database..."), AssetCount);
 
 	double StartTime = FPlatformTime::Seconds();
@@ -234,7 +729,10 @@ int32 UAssetRegistryExportCommandlet::Main(const FString& CmdLineParams)
 	const uint32 AssetsPerTransaction = 100000;
 	uint32 AssetsThisTransaction = 0;
 
-	BeginTransactionStatement.Execute();
+	if (bCSV == false)
+	{
+		BeginTransactionStatement.GetPtrOrNull()->Execute();
+	}
 	
 	// Actually add the assets.
 	bool bGotCompressedSize = false;
@@ -242,62 +740,129 @@ int32 UAssetRegistryExportCommandlet::Main(const FString& CmdLineParams)
 	{
 		FString AssetObjectPath = AssetData.GetObjectPathString();
 
-		// Add to the 'Assets' table.
+		if (bCSV)
 		{
-			GlobalInsertStatement.ClearBindings();
+			// Add to the 'Assets' table.
+			{
+				FMemoryWriter64 CSVWriter(AssetsCSV, true, true);
+
+				TUtf8StringBuilder<128> Line;
+				Line << AssetData.AssetName.ToString() << ",";
+				Line << AssetData.AssetClassPath.ToString() << ",";
+				Line << AssetObjectPath << ",";
+				
+				// If this assets has size information then add it to the global assets list as well.
+				FString CompressedSize;
+				if (AssetData.GetTagValue("Stage_ChunkCompressedSize", CompressedSize))
+				{
+					bGotCompressedSize = true;
+				}
+				Line << CompressedSize;
+				Line << "\r\n";
 
-			GlobalInsertStatement.SetBindingValueByIndex(1, AssetData.AssetName.ToString());
-			GlobalInsertStatement.SetBindingValueByIndex(2, AssetData.AssetClassPath.ToString());
-			GlobalInsertStatement.SetBindingValueByIndex(3, AssetObjectPath);
+				CSVWriter.Serialize(Line.GetData(), Line.Len());
+			}
 
-			// If this assets has size information then add it to the global assets list as well.
-			FString CompressedSize;
-			if (AssetData.GetTagValue("Stage_ChunkCompressedSize", CompressedSize))
+			// Add to our class's table.
 			{
-				bGotCompressedSize = true;
-				GlobalInsertStatement.SetBindingValueByIndex(4, CompressedSize);
-			} 
+				FClassInfo& Class = ClassInfos[AssetData.AssetClassPath];
 
-			if (GlobalInsertStatement.Execute() == false)
-			{
-				UE_LOG(LogAssetRegistryExport, Error, TEXT("Insertion Error: %s"), *OutputDatabase.GetLastError());
+				FMemoryWriter64 CSVWriter(Class.CSVMemory, true, true);
+
+				TUtf8StringBuilder<128> Line;
+				Line << AssetData.AssetName.ToString() << ",";
+				Line << AssetData.AssetClassPath.ToString() << ",";
+				Line << AssetObjectPath << ",";
+
+				// This will leave any tag values unset (NULL) in the insertion, which is what we want.
+				AssetData.EnumerateTags([&Class, &Line](TPair<FName, FAssetTagValueRef> Pair)
+				{
+					Line << Pair.Value.AsString();
+				});
+				Line << "\r\n";
+
+				CSVWriter.Serialize(Line.GetData(), Line.Len());
 			}
-		}
-		
 
-		// Add to our class's table.
+		}
+		else
 		{
-			FClassInfo& Class = ClassInfos[AssetData.AssetClassPath];
+			// Add to the 'Assets' table.
+			{
+				GlobalInsertStatement.GetPtrOrNull()->ClearBindings();
 
-			Class.InsertStatement.ClearBindings();
+				GlobalInsertStatement.GetPtrOrNull()->SetBindingValueByIndex(1, AssetData.AssetName.ToString());
+				GlobalInsertStatement.GetPtrOrNull()->SetBindingValueByIndex(2, AssetData.AssetClassPath.ToString());
+				GlobalInsertStatement.GetPtrOrNull()->SetBindingValueByIndex(3, AssetObjectPath);
 
-			Class.InsertStatement.SetBindingValueByIndex(1, AssetData.AssetName.ToString());
-			Class.InsertStatement.SetBindingValueByIndex(2, AssetData.AssetClassPath.ToString());
-			Class.InsertStatement.SetBindingValueByIndex(3, AssetObjectPath);
+				// If this assets has size information then add it to the global assets list as well.
+				FString CompressedSize;
+				if (AssetData.GetTagValue("Stage_ChunkCompressedSize", CompressedSize))
+				{
+					bGotCompressedSize = true;
+					GlobalInsertStatement.GetPtrOrNull()->SetBindingValueByIndex(4, CompressedSize);
+				} 
+
+				if (GlobalInsertStatement.GetPtrOrNull()->Execute() == false)
+				{
+					UE_LOG(LogAssetRegistryExport, Error, TEXT("Insertion Error: %s"), *OutputDatabase.GetLastError());
+				}
+			}
+		
 
-			// This will leave any tag values unset (NULL) in the insertion, which is what we want.
-			AssetData.EnumerateTags([&Class, &AssetData](TPair<FName, FAssetTagValueRef> Pair)
+			// Add to our class's table.
 			{
-				Class.InsertStatement.SetBindingValueByIndex(Class.TagBindingIndices[Pair.Key], Pair.Value.AsString());
-			});
+				FClassInfo& Class = ClassInfos[AssetData.AssetClassPath];
+
+				Class.InsertStatement.ClearBindings();
+
+				Class.InsertStatement.SetBindingValueByIndex(1, AssetData.AssetName.ToString());
+				Class.InsertStatement.SetBindingValueByIndex(2, AssetData.AssetClassPath.ToString());
+				Class.InsertStatement.SetBindingValueByIndex(3, AssetObjectPath);
 
-			if (Class.InsertStatement.Execute() == false)
+				// This will leave any tag values unset (NULL) in the insertion, which is what we want.
+				AssetData.EnumerateTags([&Class, &AssetData](TPair<FName, FAssetTagValueRef> Pair)
+				{
+					Class.InsertStatement.SetBindingValueByIndex(Class.TagBindingIndices[Pair.Key], Pair.Value.AsString());
+				});
+
+				if (Class.InsertStatement.Execute() == false)
+				{
+					UE_LOG(LogAssetRegistryExport, Error, TEXT("Class Insert Error: %s"), *OutputDatabase.GetLastError());
+				}
+			}
+
+			AssetsThisTransaction++;
+			if (AssetsThisTransaction >= AssetsPerTransaction)
 			{
-				UE_LOG(LogAssetRegistryExport, Error, TEXT("Class Insert Error: %s"), *OutputDatabase.GetLastError());
+				CommitTransactionStatement.GetPtrOrNull()->Execute();
+				BeginTransactionStatement.GetPtrOrNull()->Execute();
+				AssetsThisTransaction = 0;
 			}
 		}
 
-		AssetsThisTransaction++;
-		if (AssetsThisTransaction >= AssetsPerTransaction)
-		{
-			CommitTransactionStatement.Execute();
-			BeginTransactionStatement.Execute();
-			AssetsThisTransaction = 0;
-		}
 		return true;
 	});
 
-	CommitTransactionStatement.Execute();
+	if (bCSV == false)
+	{
+		CommitTransactionStatement.GetPtrOrNull()->Execute();
+	}
+	else
+	{
+		FArchive* Ar = IFileManager::Get().CreateFileWriter(*(CSVFileName + TEXT("Assets.csv")));
+		Ar->Serialize(AssetsCSV.GetData(), AssetsCSV.Num());
+		delete Ar;
+
+		for (TPair<FTopLevelAssetPath, FClassInfo>& Pair : ClassInfos)
+		{
+			FClassInfo& Class = Pair.Value;
+
+			Ar = IFileManager::Get().CreateFileWriter(*(CSVFileName + TEXT("/Classes/") + Class.Name + TEXT("Assets.csv")));
+			Ar->Serialize(Class.CSVMemory.GetData(), Class.CSVMemory.Num());
+			delete Ar;
+		}
+	}
 
 	double EndTime = FPlatformTime::Seconds();
 
@@ -310,22 +875,26 @@ int32 UAssetRegistryExportCommandlet::Main(const FString& CmdLineParams)
 		UE_LOG(LogAssetRegistryExport, Display, TEXT("or after staging with iostore -AssetRegistryWriteback."));
 	}
 
-	// Save a version so we can revisit if needed.
-	OutputDatabase.SetUserVersion(1);
-	
-	// Statements have to be destroyed before the db can be closed.
-	GlobalInsertStatement.Destroy();
-	CommitTransactionStatement.Destroy();
-	BeginTransactionStatement.Destroy();
-	for (TPair<FTopLevelAssetPath, FClassInfo>& Pair : ClassInfos)
+	if (bCSV == false)
 	{
-		Pair.Value.InsertStatement.Destroy();
-	}
+		// Save a version so we can revisit if needed.
+		OutputDatabase.SetUserVersion(1);
+	
+		// Statements have to be destroyed before the db can be closed.
+		GlobalInsertStatement.GetPtrOrNull()->Destroy();
+		CommitTransactionStatement.GetPtrOrNull()->Destroy();
+		BeginTransactionStatement.GetPtrOrNull()->Destroy();
+		for (TPair<FTopLevelAssetPath, FClassInfo>& Pair : ClassInfos)
+		{
+			Pair.Value.InsertStatement.Destroy();
+		}
 
-	if (OutputDatabase.Close() == false)
-	{
-		UE_LOG(LogAssetRegistryExport, Error, TEXT("Close Error: %s"), *OutputDatabase.GetLastError());
+		if (OutputDatabase.Close() == false)
+		{
+			UE_LOG(LogAssetRegistryExport, Error, TEXT("Close Error: %s"), *OutputDatabase.GetLastError());
+		}
 	}
+
 	UE_LOG(LogAssetRegistryExport, Display, TEXT("Done."));
 	return 0;
 }
