Add different size types to development asset registry.
Put the tag names in a global place to prevent string copy pasta.

diff --git a/Engine/Plugins/Editor/AssetManagerEditor/Source/AssetManagerEditor/Private/AssetManagerEditorModule.cpp b/Engine/Plugins/Editor/AssetManagerEditor/Source/AssetManagerEditor/Private/AssetManagerEditorModule.cpp
--- a/Engine/Plugins/Editor/AssetManagerEditor/Source/AssetManagerEditor/Private/AssetManagerEditorModule.cpp
+++ b/Engine/Plugins/Editor/AssetManagerEditor/Source/AssetManagerEditor/Private/AssetManagerEditorModule.cpp
@@ -109,8 +109,6 @@ const FName IAssetManagerEditorModule::ManagedDiskSizeName = FName("ManagedDiskS
 const FName IAssetManagerEditorModule::TotalUsageName = FName("TotalUsage");
 const FName IAssetManagerEditorModule::CookRuleName = FName("CookRule");
 const FName IAssetManagerEditorModule::ChunksName = FName("Chunks");
-const FName IAssetManagerEditorModule::StageChunkSizeName = FName("Stage_ChunkSize");
-const FName IAssetManagerEditorModule::StageChunkCompressedSizeName = FName("Stage_ChunkCompressedSize");
 const FName IAssetManagerEditorModule::PluginName = FName("GameFeaturePlugins");
 
 const FString FAssetManagerEditorRegistrySource::EditorSourceName = TEXT("Editor");
@@ -1371,7 +1369,7 @@ bool FAssetManagerEditorModule::GetDisplayTextForCustomColumn(const FAssetData&
 
 	UAssetManager& AssetManager = UAssetManager::Get();
 
-	if (ColumnName == ManagedResourceSizeName || ColumnName == ManagedDiskSizeName || ColumnName == DiskSizeName || ColumnName == TotalUsageName || ColumnName == StageChunkSizeName || ColumnName == StageChunkCompressedSizeName)
+	if (ColumnName == ManagedResourceSizeName || ColumnName == ManagedDiskSizeName || ColumnName == DiskSizeName || ColumnName == TotalUsageName || ColumnName == UE::AssetRegistry::Stage_ChunkSizeFName || ColumnName == UE::AssetRegistry::Stage_ChunkCompressedSizeFName)
 	{
 		// Get integer, convert to string
 		int64 IntegerValue = 0;
@@ -1489,7 +1487,7 @@ bool FAssetManagerEditorModule::GetIntegerValueForCustomColumn(const FAssetData&
 	{
 		// If the asset registry has the staged sizes written back, then use that as an accurate representation of the
 		// actual chunks in the package.
-		if (RegistrySource->GetAssetTagByObjectPath(AssetData.GetSoftObjectPath(), StageChunkSizeName, OutValue))
+		if (RegistrySource->GetAssetTagByObjectPath(AssetData.GetSoftObjectPath(), UE::AssetRegistry::Stage_ChunkSizeFName, OutValue))
 		{
 			return true;
 		}
diff --git a/Engine/Plugins/Editor/AssetManagerEditor/Source/AssetManagerEditor/Private/SAssetAuditBrowser.cpp b/Engine/Plugins/Editor/AssetManagerEditor/Source/AssetManagerEditor/Private/SAssetAuditBrowser.cpp
--- a/Engine/Plugins/Editor/AssetManagerEditor/Source/AssetManagerEditor/Private/SAssetAuditBrowser.cpp
+++ b/Engine/Plugins/Editor/AssetManagerEditor/Source/AssetManagerEditor/Private/SAssetAuditBrowser.cpp
@@ -299,7 +299,7 @@ void SAssetAuditBrowser::Construct(const FArguments& InArgs)
 	Config.CustomColumns.Emplace(FPrimaryAssetId::PrimaryAssetNameTag, LOCTEXT("AssetName", "Primary Name"), LOCTEXT("AssetNameTooltip", "Primary Asset Name of this asset, if set"), UObject::FAssetRegistryTag::TT_Alphabetical, FOnGetCustomAssetColumnData::CreateSP(this, &SAssetAuditBrowser::GetStringValueForCustomColumn), FOnGetCustomAssetColumnDisplayText::CreateSP(this, &SAssetAuditBrowser::GetDisplayTextForCustomColumn));
 	Config.CustomColumns.Emplace(IAssetManagerEditorModule::ManagedDiskSizeName, LOCTEXT("ManagedDiskSize", "Disk Size"), LOCTEXT("ManagedDiskSizeTooltip", "Total disk space used by both this and all managed assets"), UObject::FAssetRegistryTag::TT_Numerical, FOnGetCustomAssetColumnData::CreateSP(this, &SAssetAuditBrowser::GetStringValueForCustomColumn), FOnGetCustomAssetColumnDisplayText::CreateSP(this, &SAssetAuditBrowser::GetDisplayTextForCustomColumn));
 	Config.CustomColumns.Emplace(IAssetManagerEditorModule::DiskSizeName, LOCTEXT("DiskSize", "Exclusive Disk Size"), LOCTEXT("DiskSizeTooltip", "Size of saved file(s) on disk for this asset's package. If Asset Registry Writeback is enabled and a platform is selected, this will be the uncompressed size of the iostore chunks for the asset's package."), UObject::FAssetRegistryTag::TT_Numerical, FOnGetCustomAssetColumnData::CreateSP(this, &SAssetAuditBrowser::GetStringValueForCustomColumn), FOnGetCustomAssetColumnDisplayText::CreateSP(this, &SAssetAuditBrowser::GetDisplayTextForCustomColumn));
-	Config.CustomColumns.Emplace(IAssetManagerEditorModule::StageChunkCompressedSizeName, LOCTEXT("StagedCompressedSize", "Staged Compressed Size"), LOCTEXT("StagedCompressedSizeTooltip", "Compressed size of iostore chunks for this asset's package. Only visible after staging."), UObject::FAssetRegistryTag::TT_Numerical, FOnGetCustomAssetColumnData::CreateSP(this, &SAssetAuditBrowser::GetStringValueForCustomColumn), FOnGetCustomAssetColumnDisplayText::CreateSP(this, &SAssetAuditBrowser::GetDisplayTextForCustomColumn));
+	Config.CustomColumns.Emplace(UE::AssetRegistry::Stage_ChunkCompressedSizeFName, LOCTEXT("StagedCompressedSize", "Staged Compressed Size"), LOCTEXT("StagedCompressedSizeTooltip", "Compressed size of iostore chunks for this asset's package. Only visible after staging."), UObject::FAssetRegistryTag::TT_Numerical, FOnGetCustomAssetColumnData::CreateSP(this, &SAssetAuditBrowser::GetStringValueForCustomColumn), FOnGetCustomAssetColumnDisplayText::CreateSP(this, &SAssetAuditBrowser::GetDisplayTextForCustomColumn));
 	Config.CustomColumns.Emplace(IAssetManagerEditorModule::TotalUsageName, LOCTEXT("TotalUsage", "Total Usage"), LOCTEXT("TotalUsageTooltip", "Weighted count of Primary Assets that use this, higher usage means it's more likely to be in memory at runtime"), UObject::FAssetRegistryTag::TT_Numerical, FOnGetCustomAssetColumnData::CreateSP(this, &SAssetAuditBrowser::GetStringValueForCustomColumn), FOnGetCustomAssetColumnDisplayText::CreateSP(this, &SAssetAuditBrowser::GetDisplayTextForCustomColumn));
 	Config.CustomColumns.Emplace(IAssetManagerEditorModule::CookRuleName, LOCTEXT("CookRule", "Cook Rule"), LOCTEXT("CookRuleTooltip", "Whether this asset will be cooked or not"), UObject::FAssetRegistryTag::TT_Alphabetical, FOnGetCustomAssetColumnData::CreateSP(this, &SAssetAuditBrowser::GetStringValueForCustomColumn), FOnGetCustomAssetColumnDisplayText::CreateSP(this, &SAssetAuditBrowser::GetDisplayTextForCustomColumn));
 	Config.CustomColumns.Emplace(IAssetManagerEditorModule::ChunksName, LOCTEXT("Chunks", "Chunks"), LOCTEXT("ChunksTooltip", "List of chunks this will be added to when cooked"), UObject::FAssetRegistryTag::TT_Alphabetical, FOnGetCustomAssetColumnData::CreateSP(this, &SAssetAuditBrowser::GetStringValueForCustomColumn), FOnGetCustomAssetColumnDisplayText::CreateSP(this, &SAssetAuditBrowser::GetDisplayTextForCustomColumn));
diff --git a/Engine/Plugins/Editor/AssetManagerEditor/Source/AssetManagerEditor/Private/TreeView/SAssetTableTreeView.cpp b/Engine/Plugins/Editor/AssetManagerEditor/Source/AssetManagerEditor/Private/TreeView/SAssetTableTreeView.cpp
--- a/Engine/Plugins/Editor/AssetManagerEditor/Source/AssetManagerEditor/Private/TreeView/SAssetTableTreeView.cpp
+++ b/Engine/Plugins/Editor/AssetManagerEditor/Source/AssetManagerEditor/Private/TreeView/SAssetTableTreeView.cpp
@@ -1689,7 +1689,7 @@ void SAssetTableTreeView::PopulateAssetTableRow(FAssetTableRow& OutRow, const FA
 		OutRow.PrimaryName = AssetTable.StoreStr(Str);
 	}
 
-	EditorModule->GetIntegerValueForCustomColumn(AssetData, IAssetManagerEditorModule::StageChunkCompressedSizeName, OutRow.StagedCompressedSize, &RegistrySource);
+	EditorModule->GetIntegerValueForCustomColumn(AssetData, UE::AssetRegistry::Stage_ChunkCompressedSizeFName, OutRow.StagedCompressedSize, &RegistrySource);
 	EditorModule->GetIntegerValueForCustomColumn(AssetData, IAssetManagerEditorModule::TotalUsageName, OutRow.TotalUsageCount, &RegistrySource);
 
 	if (EditorModule->GetStringValueForCustomColumn(AssetData, IAssetManagerEditorModule::ChunksName, Str, &RegistrySource))
diff --git a/Engine/Plugins/Editor/AssetManagerEditor/Source/AssetManagerEditor/Public/AssetManagerEditorModule.h b/Engine/Plugins/Editor/AssetManagerEditor/Source/AssetManagerEditor/Public/AssetManagerEditorModule.h
--- a/Engine/Plugins/Editor/AssetManagerEditor/Source/AssetManagerEditor/Public/AssetManagerEditorModule.h
+++ b/Engine/Plugins/Editor/AssetManagerEditor/Source/AssetManagerEditor/Public/AssetManagerEditorModule.h
@@ -381,8 +381,6 @@ public:
 	static const FName TotalUsageName;
 	static const FName CookRuleName;
 	static const FName ChunksName;
-	static const FName StageChunkSizeName;
-	static const FName StageChunkCompressedSizeName;
 	static const FName PluginName;
 
 	/** Gets the value of a "virtual" column for an asset data, this will query the AssetManager for you and takes current platform into account. Returns true and sets out parameter if found */
diff --git a/Engine/Plugins/Editor/AssetRegistryExport/Source/Private/AssetRegistryExportCommandlet.cpp b/Engine/Plugins/Editor/AssetRegistryExport/Source/Private/AssetRegistryExportCommandlet.cpp
--- a/Engine/Plugins/Editor/AssetRegistryExport/Source/Private/AssetRegistryExportCommandlet.cpp
+++ b/Engine/Plugins/Editor/AssetRegistryExport/Source/Private/AssetRegistryExportCommandlet.cpp
@@ -274,11 +274,11 @@ int32 UAssetRegistryExportCommandlet::Main(const FString& CmdLineParams)
 		TSet<const FAssetData*> RemainingAssetsWithSize;
 		uint64 TotalCompressedBytes = 0;
 		{
-			const TArray<const FAssetData*>& AssetsWithSize = AssetRegistry.GetAssetsByTagName("Stage_ChunkCompressedSize");
+			const TArray<const FAssetData*>& AssetsWithSize = AssetRegistry.GetAssetsByTagName(UE::AssetRegistry::Stage_ChunkCompressedSizeFName);
 			for (const FAssetData* AD : AssetsWithSize)
 			{
 				uint64 CompressedSize = 0;
-				if (AD->GetTagValue("Stage_ChunkCompressedSize", CompressedSize))
+				if (AD->GetTagValue(UE::AssetRegistry::Stage_ChunkCompressedSizeFName, CompressedSize))
 				{
 					TotalCompressedBytes += CompressedSize;
 					RemainingAssetsWithSize.Add(AD);
@@ -369,7 +369,7 @@ int32 UAssetRegistryExportCommandlet::Main(const FString& CmdLineParams)
 						Line << GetRootClassToPrint(UseAsClassData->AssetClassPath) << ",";
 
 						uint64 CompressedSize = 0;
-						if (UseAsClassData->GetTagValue("Stage_ChunkCompressedSize", CompressedSize))
+						if (UseAsClassData->GetTagValue(UE::AssetRegistry::Stage_ChunkCompressedSizeFName, CompressedSize))
 						{
 							Info.SizeOfAllDependencies += CompressedSize;
 						}
@@ -428,7 +428,7 @@ int32 UAssetRegistryExportCommandlet::Main(const FString& CmdLineParams)
 						if (UseAsClassData)
 						{
 							uint64 CompressedSize = 0;
-							if (UseAsClassData->GetTagValue("Stage_ChunkCompressedSize", CompressedSize))
+							if (UseAsClassData->GetTagValue(UE::AssetRegistry::Stage_ChunkCompressedSizeFName, CompressedSize))
 							{
 								check(RemainingAssetsWithSize.Contains(UseAsClassData));
 								RemainingAssetsWithSize.Remove(UseAsClassData);
@@ -464,7 +464,7 @@ int32 UAssetRegistryExportCommandlet::Main(const FString& CmdLineParams)
 						Line << GetRootClassToPrint(UseAsClassData->AssetClassPath) << ",";
 
 						uint64 CompressedSize = 0;
-						if (UseAsClassData->GetTagValue("Stage_ChunkCompressedSize", CompressedSize))
+						if (UseAsClassData->GetTagValue(UE::AssetRegistry::Stage_ChunkCompressedSizeFName, CompressedSize))
 						{
 							check(RemainingAssetsWithSize.Contains(UseAsClassData));
 							RemainingAssetsWithSize.Remove(UseAsClassData);
@@ -505,7 +505,7 @@ int32 UAssetRegistryExportCommandlet::Main(const FString& CmdLineParams)
 						Line << GetRootClassToPrint(UseAsClassData->AssetClassPath) << ",";
 
 						uint64 CompressedSize = 0;
-						if (UseAsClassData->GetTagValue("Stage_ChunkCompressedSize", CompressedSize))
+						if (UseAsClassData->GetTagValue(UE::AssetRegistry::Stage_ChunkCompressedSizeFName, CompressedSize))
 						{
 							check(RemainingAssetsWithSize.Contains(UseAsClassData));
 							RemainingAssetsWithSize.Remove(UseAsClassData);
@@ -605,7 +605,7 @@ int32 UAssetRegistryExportCommandlet::Main(const FString& CmdLineParams)
 		for (const FAssetData* LD : RemainingAssetsWithSize)
 		{
 			uint64 value = 0;
-			if (LD->GetTagValue("Stage_ChunkCompressedSize", value))
+			if (LD->GetTagValue(UE::AssetRegistry::Stage_ChunkCompressedSizeFName, value))
 			{
 				const FAssetPackageData* PackageData = AllPackages[LD->PackageName];
 
@@ -615,7 +615,7 @@ int32 UAssetRegistryExportCommandlet::Main(const FString& CmdLineParams)
 				for (const FAssetData* AD : PackageAssets)
 				{
 					uint64 blah = 0;
-					AD->GetTagValue("Stage_ChunkCompressedSize", blah);
+					AD->GetTagValue(UE::AssetRegistry::Stage_ChunkCompressedSizeFName, blah);
 
 					UE_LOG(LogAssetRegistryExport, Warning, TEXT("......Package has %s %d"), *AD->AssetName.ToString(), blah);
 
@@ -753,7 +753,7 @@ int32 UAssetRegistryExportCommandlet::Main(const FString& CmdLineParams)
 				
 				// If this assets has size information then add it to the global assets list as well.
 				FString CompressedSize;
-				if (AssetData.GetTagValue("Stage_ChunkCompressedSize", CompressedSize))
+				if (AssetData.GetTagValue(UE::AssetRegistry::Stage_ChunkCompressedSizeFName, CompressedSize))
 				{
 					bGotCompressedSize = true;
 				}
@@ -797,7 +797,7 @@ int32 UAssetRegistryExportCommandlet::Main(const FString& CmdLineParams)
 
 				// If this assets has size information then add it to the global assets list as well.
 				FString CompressedSize;
-				if (AssetData.GetTagValue("Stage_ChunkCompressedSize", CompressedSize))
+				if (AssetData.GetTagValue(UE::AssetRegistry::Stage_ChunkCompressedSizeFName, CompressedSize))
 				{
 					bGotCompressedSize = true;
 					GlobalInsertStatement.GetPtrOrNull()->SetBindingValueByIndex(4, CompressedSize);
diff --git a/Engine/Source/Developer/IoStoreUtilities/Private/IoStoreUtilities.cpp b/Engine/Source/Developer/IoStoreUtilities/Private/IoStoreUtilities.cpp
--- a/Engine/Source/Developer/IoStoreUtilities/Private/IoStoreUtilities.cpp
+++ b/Engine/Source/Developer/IoStoreUtilities/Private/IoStoreUtilities.cpp
@@ -3493,20 +3493,23 @@ static void AddChunkInfoToAssetRegistry(TMap<FPackageId, TArray<FIoStoreChunkSou
 			continue;
 		}
 
+		UE::Cook::FPluginSizeInfo PackageCompressedSize;
+		UE::Cook::FPluginSizeInfo PackageSize;
 		int32 ChunkCount = 0;
-		int64 Size = 0;
-		int64 CompressedSize = 0;
 		for (const FIoStoreChunkSource& ChunkInfo : *PackageChunks)
 		{
 			ChunkCount++;
-			Size += ChunkInfo.ChunkInfo.Size;
-			CompressedSize += ChunkInfo.ChunkInfo.CompressedSize;
+			PackageSize[ChunkInfo.SizeType] += ChunkInfo.ChunkInfo.Size;
+			PackageCompressedSize[ChunkInfo.SizeType] += ChunkInfo.ChunkInfo.CompressedSize;
 		}
 
 		FAssetDataTagMap TagsAndValues;
-		TagsAndValues.Add("Stage_ChunkCount", LexToString(ChunkCount));
-		TagsAndValues.Add("Stage_ChunkSize", LexToString(Size));
-		TagsAndValues.Add("Stage_ChunkCompressedSize", LexToString(CompressedSize));
+		TagsAndValues.Add(UE::AssetRegistry::Stage_ChunkCountFName, LexToString(ChunkCount));
+		TagsAndValues.Add(UE::AssetRegistry::Stage_ChunkSizeFName, LexToString(PackageSize.TotalSize()));
+		TagsAndValues.Add(UE::AssetRegistry::Stage_ChunkCompressedSizeFName, LexToString(PackageCompressedSize.TotalSize()));
+		TagsAndValues.Add(UE::AssetRegistry::Stage_ChunkInstalledSizeFName, LexToString(PackageCompressedSize[UE::Cook::EPluginSizeTypes::Installed]));
+		TagsAndValues.Add(UE::AssetRegistry::Stage_ChunkStreamingSizeFName, LexToString(PackageCompressedSize[UE::Cook::EPluginSizeTypes::Streaming]));
+		TagsAndValues.Add(UE::AssetRegistry::Stage_ChunkOptionalSizeFName, LexToString(PackageCompressedSize[UE::Cook::EPluginSizeTypes::Optional]));
 		AssetRegistry.AddTagsToAssetData(AssetData->GetSoftObjectPath(), MoveTemp(TagsAndValues));
 
 		// We assign a package's chunks to a single asset, remove it from the list so that
@@ -3514,7 +3517,7 @@ static void AddChunkInfoToAssetRegistry(TMap<FPackageId, TArray<FIoStoreChunkSou
 		PackageToChunks.Remove(FPackageId::FromName(AssetPackage.Key));
 
 		UpdatedAssetCount++;
-		AssetsCompressedSize += CompressedSize;
+		AssetsCompressedSize += PackageCompressedSize.TotalSize();
 	}
 	
 	// PackageToChunks now has chunks that we never assigned to an asset, and so aren't accounted for.
diff --git a/Engine/Source/Editor/UnrealEd/Private/Commandlets/AssetSizeQueryCommandlet.cpp b/Engine/Source/Editor/UnrealEd/Private/Commandlets/AssetSizeQueryCommandlet.cpp
--- a/Engine/Source/Editor/UnrealEd/Private/Commandlets/AssetSizeQueryCommandlet.cpp
+++ b/Engine/Source/Editor/UnrealEd/Private/Commandlets/AssetSizeQueryCommandlet.cpp
@@ -145,7 +145,7 @@ int32 UAssetSizeQueryCommandlet::Main(const FString& FullCommandLine)
 		 &AssetFilter](const FAssetData& AssetData)
 	{
 		FString CompressedSize;
-		if (AssetData.GetTagValue("Stage_ChunkCompressedSize", CompressedSize) == false ||
+		if (AssetData.GetTagValue(UE::AssetRegistry::Stage_ChunkCompressedSizeFName, CompressedSize) == false ||
 			CompressedSize.Len() == 0)
 		{
 			return true;
diff --git a/Engine/Source/Runtime/AssetRegistry/Private/AssetRegistry.cpp b/Engine/Source/Runtime/AssetRegistry/Private/AssetRegistry.cpp
--- a/Engine/Source/Runtime/AssetRegistry/Private/AssetRegistry.cpp
+++ b/Engine/Source/Runtime/AssetRegistry/Private/AssetRegistry.cpp
@@ -107,6 +107,13 @@ namespace UE::AssetRegistry
 {
 	const FName WildcardFName(TEXT("*"));
 	const FTopLevelAssetPath WildcardPathName(TEXT("/*"), TEXT("*"));
+
+	const FName Stage_ChunkCountFName(TEXT("Stage_ChunkCount"));
+	const FName Stage_ChunkSizeFName(TEXT("Stage_ChunkSize"));
+	const FName Stage_ChunkCompressedSizeFName(TEXT("Stage_ChunkCompressedSize"));
+	const FName Stage_ChunkInstalledSizeFName(TEXT("Stage_ChunkInstalledSize"));
+	const FName Stage_ChunkStreamingSizeFName(TEXT("Stage_ChunkStreamingSize"));
+	const FName Stage_ChunkOptionalSizeFName(TEXT("Stage_ChunkOptionalSize"));
 }
 
 namespace UE::AssetRegistry::Impl
diff --git a/Engine/Source/Runtime/AssetRegistry/Public/AssetRegistry/IAssetRegistry.h b/Engine/Source/Runtime/AssetRegistry/Public/AssetRegistry/IAssetRegistry.h
--- a/Engine/Source/Runtime/AssetRegistry/Public/AssetRegistry/IAssetRegistry.h
+++ b/Engine/Source/Runtime/AssetRegistry/Public/AssetRegistry/IAssetRegistry.h
@@ -923,6 +923,30 @@ namespace AssetRegistry
 	extern ASSETREGISTRY_API const FName WildcardFName;
 	extern ASSETREGISTRY_API const FTopLevelAssetPath WildcardPathName;
 
+
+	/*
+	* Various FNames for asset tags that get added during staging if asset registry writeback is enabled.
+	* Note that these sizes refer to the package as a whole, and some packages have more than one asset per
+	* package. In that case, the data is stored on the asset returned by GetMostImportantAsset with IgnoreSkipClasses.
+	* 
+	* NOTE some platforms do not compress in UnrealPak. For those platforms, all compressed sizes are actually
+	* uncompressed!
+	* 
+	* Stage_ChunkCountFName						Total iostore chunks in the package.
+	* Stage_ChunkSizeFName						The total uncompressed size of the chunks in the package.
+	* Stage_ChunkCompressedSizeFName			The total compressed size of the chunks in the package. SEE NOTE ABOVE!
+	* Stage_ChunkInstalledSizeFName				The compressed size of all chunks that must be installed with the game/plugin.
+	* Stage_ChunkStreamingSizeFName				The compressed size of all chunks that are delivered using IAS.
+	* Stage_ChunkOptionalSizeFName				The compressed size of all chunks that are placed in an optional container.
+	* 
+	* See CookMetadata.h - UE::Cook::EPluginSizeTypes for more information on size types.
+	*/
+	extern ASSETREGISTRY_API const FName Stage_ChunkCountFName;
+	extern ASSETREGISTRY_API const FName Stage_ChunkSizeFName;
+	extern ASSETREGISTRY_API const FName Stage_ChunkCompressedSizeFName;
+	extern ASSETREGISTRY_API const FName Stage_ChunkInstalledSizeFName;
+	extern ASSETREGISTRY_API const FName Stage_ChunkStreamingSizeFName;
+	extern ASSETREGISTRY_API const FName Stage_ChunkOptionalSizeFName;
 } // namespace AssetRegistry
 } // namespace UE
 
