[AssetRegistryExportCommandlet] Add quotes around tag names to prevent sql errors when the column/tag names are also sql keywords.

diff --git a/Engine/Plugins/Editor/AssetRegistryExport/Source/Private/AssetRegistryExportCommandlet.cpp b/Engine/Plugins/Editor/AssetRegistryExport/Source/Private/AssetRegistryExportCommandlet.cpp
--- a/Engine/Plugins/Editor/AssetRegistryExport/Source/Private/AssetRegistryExportCommandlet.cpp
+++ b/Engine/Plugins/Editor/AssetRegistryExport/Source/Private/AssetRegistryExportCommandlet.cpp
@@ -672,8 +672,8 @@ int32 UAssetRegistryExportCommandlet::Main(const FString& CmdLineParams)
 			{
 				FString& SanitizedTagName = Class.SanitizedTagNames[TagName];
 
-				CreateTable += TEXT(", ") + SanitizedTagName;
-				InsertStatement += TEXT(", ") + SanitizedTagName;
+				CreateTable += TEXT(", \"") + SanitizedTagName + TEXT("\"");
+				InsertStatement += TEXT(", \"") + SanitizedTagName + TEXT("\"");
 				ValuesStatement += FString::Printf(TEXT(", ?%d"), CurrentIndex);
 
 				Class.TagBindingIndices.Add(TagName, CurrentIndex);
