Update AssetRegistryExport:
* Add new chunk size tags to global assets table
* Allow for filtering to a single class type to cut down on tables/db size
* Fix CSV to ensure values are in the right column order.

#rnx

diff --git a/Engine/Plugins/Editor/AssetRegistryExport/Source/Private/AssetRegistryExportCommandlet.cpp b/Engine/Plugins/Editor/AssetRegistryExport/Source/Private/AssetRegistryExportCommandlet.cpp
--- a/Engine/Plugins/Editor/AssetRegistryExport/Source/Private/AssetRegistryExportCommandlet.cpp
+++ b/Engine/Plugins/Editor/AssetRegistryExport/Source/Private/AssetRegistryExportCommandlet.cpp
@@ -64,6 +64,8 @@ int32 UAssetRegistryExportCommandlet::Main(const FString& CmdLineParams)
 		UE_LOG(LogAssetRegistryExport, Display, TEXT("                                        a commandlet for the project that generated the asset registry."));
 		UE_LOG(LogAssetRegistryExport, Display, TEXT("    -IncludeSharedWithDependencies      When listing dependencies, this will also add shared dependencies"));
 		UE_LOG(LogAssetRegistryExport, Display, TEXT("                                        to a \"shared\" root."));
+		UE_LOG(LogAssetRegistryExport, Display, TEXT("    -FilterToClass=<classname>          Only add the table for the given class (and global assets)"));
+		UE_LOG(LogAssetRegistryExport, Display, TEXT("                                        Example: -FilterToClass=\"/Script/Engine.Texture2D\""));
 		return 1;
 	}
 
@@ -151,6 +153,12 @@ int32 UAssetRegistryExportCommandlet::Main(const FString& CmdLineParams)
 		OutputDatabase.Execute(TEXT("PRAGMA locking_mode=EXCLUSIVE;"));
 	}
 
+	FString FilterToClass;
+	FParse::Value(*CmdLineParams, TEXT("FilterToClass="), FilterToClass);
+	if (FilterToClass.Len())
+	{
+		UE_LOG(LogAssetRegistryExport, Display, TEXT("Got option: FilterToClass=%s"), *FilterToClass);
+	}
 
 	struct FClassInfo
 	{
@@ -175,8 +183,16 @@ int32 UAssetRegistryExportCommandlet::Main(const FString& CmdLineParams)
 	// given class, so we have to get the union of all possibles.
 	AssetRegistry.EnumerateAllAssets(TSet<FName>(), [&](const FAssetData& AssetData)
 	{
+		FString AssetClassPathString = AssetData.AssetClassPath.ToString();
+
+		if (FilterToClass.Len() &&
+			FilterToClass != AssetClassPathString)
+		{
+			return true;
+		}
+
 		FClassInfo& ClassInfo = ClassInfos.FindOrAdd(AssetData.AssetClassPath);
-		ClassInfo.Name = AssetData.AssetClassPath.ToString();
+		ClassInfo.Name = MoveTemp(AssetClassPathString);
 		{
 			auto& CharArray = ClassInfo.Name.GetCharArray();
 			for (int i = 0; i < CharArray.Num() - 1; i++)
@@ -707,20 +723,20 @@ int32 UAssetRegistryExportCommandlet::Main(const FString& CmdLineParams)
 		FMemoryWriter64 CSVWriter(AssetsCSV);
 		
 		TUtf8StringBuilder<128> Header;
-		Header << "Name,Class,Path,Stage_ChunkCompressedSize\r\n";
+		Header << "Name,Class,Path,Stage_ChunkCompressedSize,Stage_ChunkSize,Stage_ChunkOptionalSize,Stage_ChunkStreamingSize\r\n";
 		CSVWriter.Serialize(Header.GetData(), Header.Len());
 	}
 	else
 	{
 		// Create catchall table and insert statements
-		if (OutputDatabase.Execute(TEXT("CREATE TABLE IF NOT EXISTS Assets(Id INTEGER PRIMARY KEY, Name, Class, Path TEXT UNIQUE, Stage_ChunkCompressedSize INTEGER);")) == false)
+		if (OutputDatabase.Execute(TEXT("CREATE TABLE IF NOT EXISTS Assets(Id INTEGER PRIMARY KEY, Name, Class, Path TEXT UNIQUE, Stage_ChunkCompressedSize INTEGER, Stage_ChunkSize INTEGER, Stage_ChunkOptionalSize INTEGER, Stage_ChunkStreamingSize INTEGER);")) == false)
 		{
 			UE_LOG(LogAssetRegistryExport, Error, TEXT("Error: %s"), *OutputDatabase.GetLastError());
 			OutputDatabase.Close();
 			return 1;
 		}
 
-		GlobalInsertStatement.Emplace(OutputDatabase, TEXT("INSERT INTO Assets(Name, Class, Path, Stage_ChunkCompressedSize) Values(?1, ?2, ?3, ?4);"));
+		GlobalInsertStatement.Emplace(OutputDatabase, TEXT("INSERT INTO Assets(Name, Class, Path, Stage_ChunkCompressedSize, Stage_ChunkSize, Stage_ChunkOptionalSize, Stage_ChunkStreamingSize) Values(?1, ?2, ?3, ?4, ?5, ?6, ?7);"));
 		if (GlobalInsertStatement.GetPtrOrNull()->IsValid() == false)
 		{
 			UE_LOG(LogAssetRegistryExport, Error, TEXT("Error: %s"), *OutputDatabase.GetLastError());
@@ -759,9 +775,9 @@ int32 UAssetRegistryExportCommandlet::Main(const FString& CmdLineParams)
 				FMemoryWriter64 CSVWriter(AssetsCSV, true, true);
 
 				TUtf8StringBuilder<128> Line;
-				Line << AssetData.AssetName.ToString() << ",";
-				Line << AssetData.AssetClassPath.ToString() << ",";
-				Line << AssetObjectPath << ",";
+				Line << AssetData.AssetName.ToString() << TEXTVIEW(",");
+				Line << AssetData.AssetClassPath.ToString() << TEXTVIEW(",");
+				Line << AssetObjectPath << TEXTVIEW(",");
 				
 				// If this assets has size information then add it to the global assets list as well.
 				FString CompressedSize;
@@ -769,28 +785,43 @@ int32 UAssetRegistryExportCommandlet::Main(const FString& CmdLineParams)
 				{
 					bGotCompressedSize = true;
 				}
-				Line << CompressedSize;
-				Line << "\r\n";
+				FString Size;
+				AssetData.GetTagValue(UE::AssetRegistry::Stage_ChunkSizeFName, Size);
+				FString StreamingSize;
+				AssetData.GetTagValue(UE::AssetRegistry::Stage_ChunkStreamingSizeFName, StreamingSize);
+				FString OptionalSize;
+				AssetData.GetTagValue(UE::AssetRegistry::Stage_ChunkOptionalSizeFName, OptionalSize);
+				
+				Line << CompressedSize << TEXTVIEW(",");
+				Line << Size << TEXTVIEW(",");
+				Line << OptionalSize << TEXTVIEW(",");
+				Line << StreamingSize;
+
+				Line << TEXTVIEW("\r\n");
 
 				CSVWriter.Serialize(Line.GetData(), Line.Len());
 			}
 
 			// Add to our class's table.
+			FClassInfo* ClassPtr = ClassInfos.Find(AssetData.AssetClassPath);
+			if (ClassPtr)
 			{
-				FClassInfo& Class = ClassInfos[AssetData.AssetClassPath];
+				FClassInfo& Class = *ClassPtr;
 
 				FMemoryWriter64 CSVWriter(Class.CSVMemory, true, true);
 
 				TUtf8StringBuilder<128> Line;
 				Line << AssetData.AssetName.ToString() << ",";
 				Line << AssetData.AssetClassPath.ToString() << ",";
-				Line << AssetObjectPath << ",";
+				Line << AssetObjectPath;
 
-				// This will leave any tag values unset (NULL) in the insertion, which is what we want.
-				AssetData.EnumerateTags([&Class, &Line](TPair<FName, FAssetTagValueRef> Pair)
+				for (FName TagName : Class.Tags)
 				{
-					Line << Pair.Value.AsString();
-				});
+					FString Value;
+					AssetData.GetTagValue(TagName, Value);
+					Line << "," << Value;
+				}
+
 				Line << "\r\n";
 
 				CSVWriter.Serialize(Line.GetData(), Line.Len());
@@ -815,6 +846,23 @@ int32 UAssetRegistryExportCommandlet::Main(const FString& CmdLineParams)
 					GlobalInsertStatement.GetPtrOrNull()->SetBindingValueByIndex(4, CompressedSize);
 				} 
 
+				FString Size;
+				if (AssetData.GetTagValue(UE::AssetRegistry::Stage_ChunkSizeFName, Size))
+				{
+					GlobalInsertStatement.GetPtrOrNull()->SetBindingValueByIndex(5, Size);
+				}
+
+				FString StreamingSize;
+				if (AssetData.GetTagValue(UE::AssetRegistry::Stage_ChunkStreamingSizeFName, StreamingSize))
+				{
+					GlobalInsertStatement.GetPtrOrNull()->SetBindingValueByIndex(6, StreamingSize);
+				}
+				FString OptionalSize;
+				if (AssetData.GetTagValue(UE::AssetRegistry::Stage_ChunkOptionalSizeFName, OptionalSize))
+				{
+					GlobalInsertStatement.GetPtrOrNull()->SetBindingValueByIndex(7, OptionalSize);
+				}
+
 				if (GlobalInsertStatement.GetPtrOrNull()->Execute() == false)
 				{
 					UE_LOG(LogAssetRegistryExport, Error, TEXT("Insertion Error: %s"), *OutputDatabase.GetLastError());
@@ -823,8 +871,10 @@ int32 UAssetRegistryExportCommandlet::Main(const FString& CmdLineParams)
 		
 
 			// Add to our class's table.
+			FClassInfo* ClassPtr = ClassInfos.Find(AssetData.AssetClassPath);
+			if (ClassPtr)
 			{
-				FClassInfo& Class = ClassInfos[AssetData.AssetClassPath];
+				FClassInfo& Class = *ClassPtr;
 
 				Class.InsertStatement.ClearBindings();
 
