Don't try to add to an existing database when exporting the AR to SQL.

diff --git a/Engine/Plugins/Editor/AssetRegistryExport/Source/Private/AssetRegistryExportCommandlet.cpp b/Engine/Plugins/Editor/AssetRegistryExport/Source/Private/AssetRegistryExportCommandlet.cpp
--- a/Engine/Plugins/Editor/AssetRegistryExport/Source/Private/AssetRegistryExportCommandlet.cpp
+++ b/Engine/Plugins/Editor/AssetRegistryExport/Source/Private/AssetRegistryExportCommandlet.cpp
@@ -75,6 +75,12 @@ int32 UAssetRegistryExportCommandlet::Main(const FString& CmdLineParams)
 
 	if (bCSV == false)
 	{
+		// We get a mountain of errors if we append to a database.
+		if (!IFileManager::Get().Delete(*OutputDatabaseName, false))
+		{
+			UE_LOG(LogAssetRegistryExport, Error, TEXT("Failed to delete sqlite database before opening %s"), *OutputDatabaseName);
+			return 1;
+		}
 		if (OutputDatabase.Open(*OutputDatabaseName) == false)
 		{
 			UE_LOG(LogAssetRegistryExport, Error, TEXT("Failed to open sqlite database %s"), *OutputDatabaseName);
